PR #7785: Re- implemented FastBuild support for shader compilation as a Plugin (Contributed by [REDACTED_PERSON]).

- Also removed the previous FastBuild shader compiler integration based on an older API.
- Made Distributed Build controller selection OS-agnostic.

diff --git a/Engine/Plugins/FastBuildController/FastBuildController.uplugin b/Engine/Plugins/FastBuildController/FastBuildController.uplugin
new file mode 100644
--- /dev/null
+++ b/Engine/Plugins/FastBuildController/FastBuildController.uplugin
@@ -0,0 +1,31 @@
+﻿{
+	"FileVersion" : 1,
+	"Version" : 1,
+	"VersionName" : "0.1",
+	"FriendlyName" : "FastBuild Controller",
+	"Description" : "Adds support for shader compiling distribution using FastBuild",
+	"Category" : "Build Distribution",
+	"CreatedBy" : "Epic Games, Inc.",
+	"CreatedByURL" : "http://epicgames.com",
+	"DocsURL" : "",
+	"MarketplaceURL" : "",
+	"SupportURL" : "",
+	"EnabledByDefault" : true,
+	"CanContainContent" : false,
+	"IsBetaVersion" : false,
+	"Installed" : false,
+
+	"Modules" :
+	[
+		{
+			"Name" : "FastBuildController",
+			"Type" : "UncookedOnly",
+			"LoadingPhase" : "EarliestPossible",
+			"WhitelistPlatforms": [
+				"Win64",
+				"Mac",
+				"Linux"
+			]
+		}
+	]
+}
diff --git a/Engine/Plugins/FastBuildController/Source/FastbuildController.Build.cs b/Engine/Plugins/FastBuildController/Source/FastbuildController.Build.cs
new file mode 100644
--- /dev/null
+++ b/Engine/Plugins/FastBuildController/Source/FastbuildController.Build.cs
@@ -0,0 +1,18 @@
+// Copyright Epic Games, Inc. All Rights Reserved.
+
+using UnrealBuildTool;
+
+public class FastBuildController : ModuleRules
+{
+	public FastBuildController(ReadOnlyTargetRules TargetRules)
+		: base(TargetRules)
+	{
+		PrivateDependencyModuleNames.AddRange(new string[]
+		{
+			"Core",
+			"RenderCore"
+		});
+		
+		PrivateIncludePathModuleNames.Add("Engine");
+	}
+}
diff --git a/Engine/Plugins/FastBuildController/Source/Private/FastBuildControllerModule.cpp b/Engine/Plugins/FastBuildController/Source/Private/FastBuildControllerModule.cpp
new file mode 100644
--- /dev/null
+++ b/Engine/Plugins/FastBuildController/Source/Private/FastBuildControllerModule.cpp
@@ -0,0 +1,331 @@
+// Copyright Epic Games, Inc. All Rights Reserved.
+
+#include "FastBuildControllerModule.h"
+#include "Features/IModularFeatures.h"
+#include "HAL/FileManager.h"
+#include "HAL/PlatformFileManager.h"
+#include "HAL/PlatformFile.h"
+#include "Misc/ScopeLock.h"
+#include "Misc/Paths.h"
+#include "Misc/Guid.h"
+#include "Modules/ModuleManager.h"
+#include "FastBuildJobProcessor.h"
+#include "FastBuildUtilities.h"
+
+DEFINE_LOG_CATEGORY_STATIC(LogFastBuildController, Log, Log);
+
+namespace FASTBuildControllerVariables
+{
+	int32 Enabled = 1;
+	FAutoConsoleVariableRef CVarFASTBuildControllerShaderCompile(
+        TEXT("r.FASTBuildController.Enabled"),
+        Enabled,
+        TEXT("Enables or disables the use of FASTBuild to build shaders.\n")
+        TEXT("0: Controller will not be used (shaders will be built locally or using other controllers). \n")
+        TEXT("1: Distribute builds using FASTBuild."),
+        ECVF_Default);
+}
+
+#if PLATFORM_MAC
+static FString GetMetalCompilerFolder()
+{
+	FString Result;
+	if (FPlatformProcess::ExecProcess(TEXT("/usr/bin/xcrun"), TEXT("--sdk macosx metal -v --target=air64-apple-darwin19.6.0"), nullptr, &Result, &Result))
+	{
+		const TCHAR InstalledDirText[] = TEXT("InstalledDir:");
+		int32 InstalledDirOffset = Result.Find(InstalledDirText, ESearchCase::CaseSensitive);
+		if (InstalledDirOffset != INDEX_NONE)
+		{
+			InstalledDirOffset += UE_ARRAY_COUNT(InstalledDirText);
+			const int32 MacOSBinOffset = Result.Find(TEXT("/macos/bin\n"), ESearchCase::CaseSensitive, ESearchDir::FromStart, InstalledDirOffset);
+			if (MacOSBinOffset != INDEX_NONE)
+			{
+				FString Substring = Result.Mid(InstalledDirOffset, MacOSBinOffset - InstalledDirOffset);
+				return Result.Mid(InstalledDirOffset, MacOSBinOffset - InstalledDirOffset);
+			}
+		}
+	}
+
+	return FString();
+}
+#endif // PLATFORM_MAC
+
+FFastBuildControllerModule::FFastBuildControllerModule()
+	: bSupported(false)
+	, bModuleInitialized(false)
+	, bControllerInitialized(false)
+	, TasksCS(new FCriticalSection())
+	, RootWorkingDirectory(FString::Printf(TEXT("%sUnrealFastBuildWorkingDir/"), FPlatformProcess::UserTempDir()))
+	, WorkingDirectory(RootWorkingDirectory + FGuid::NewGuid().ToString(EGuidFormats::Digits))
+	, bShutdown(false)
+{}
+
+FFastBuildControllerModule::~FFastBuildControllerModule()
+{
+	if (JobDispatcherThread)
+	{
+		JobDispatcherThread->Stop();
+	}
+	
+	CleanWorkingDirectory();
+}
+
+bool FFastBuildControllerModule::IsSupported()
+{
+	if (bControllerInitialized)
+	{
+		return bSupported;
+	}
+
+	if (FParse::Param(FCommandLine::Get(), TEXT("nofastbuildcontroller")) ||
+		FParse::Param(FCommandLine::Get(), TEXT("nofastbuildshadercompile")) ||
+		FParse::Param(FCommandLine::Get(), TEXT("nofbuildshadercompile")) ||
+		FParse::Param(FCommandLine::Get(), TEXT("noshaderworker")))
+	{
+		FASTBuildControllerVariables::Enabled = 0;
+	}
+
+	// Check to see if the FASTBuild exe exists.
+	if (FASTBuildControllerVariables::Enabled == 1)
+	{
+		const FString BrokeragePath = FPlatformMisc::GetEnvironmentVariable(TEXT("FASTBUILD_BROKERAGE_PATH"));
+		if (BrokeragePath.IsEmpty())
+		{
+			FASTBuildControllerVariables::Enabled = 0;
+			return false;
+		}
+
+		IPlatformFile& PlatformFile = FPlatformFileManager::Get().GetPlatformFile();
+
+		const FString FASTBuildAbsolutePath = IFileManager::Get().ConvertToAbsolutePathForExternalAppForWrite(*FastBuildUtilities::GetFastBuildExecutablePath());
+		if (!PlatformFile.FileExists(*FASTBuildAbsolutePath))
+		{
+			UE_LOG(LogFastBuildController, Warning, TEXT("Cannot use FASTBuild Shader Compiler as FASTBuild is not found: %s"),
+				*FPaths::ConvertRelativePathToFull(FASTBuildAbsolutePath));
+			FASTBuildControllerVariables::Enabled = 0;
+			return false;
+		}	
+
+#if PLATFORM_MAC
+		static bool bCopyMetalCompilerToIntermediateDir = true;
+		if (bCopyMetalCompilerToIntermediateDir)
+		{
+			SCOPED_AUTORELEASE_POOL;
+
+			// Make a copy of all the Metal shader compiler files in the intermediate folder, so that they are in the same directory tree as SharedCompileWorker.
+			// This is required for FASTBuild to preserve the directory structure when it copies these files to the worker
+			const FString SrcDir = GetMetalCompilerFolder();
+			if (SrcDir.Len() == 0)
+			{
+				UE_LOG(LogFastBuildController, Warning, TEXT("Cannot use FASTBuild Shader Compiler as Metal shader compiler could not be found"));
+				FASTBuildControllerVariables::Enabled = 0;
+			}
+			else
+			{
+				const FString IntermediateShadersDir = FPaths::EngineIntermediateDir() / TEXT("Shaders");
+				const FString DestDir = IntermediateShadersDir / TEXT("metal");
+				if (PlatformFile.DirectoryExists(*DestDir))
+				{
+					PlatformFile.DeleteDirectoryRecursively(*DestDir);
+				}
+
+				if (!PlatformFile.DirectoryExists(*IntermediateShadersDir))
+				{
+					PlatformFile.CreateDirectoryTree(*IntermediateShadersDir);
+				}
+
+				NSFileManager* FileManager = [NSFileManager defaultManager]; // Use NSFileManager as PlatformFile's CopyDirectoryTree does not preserve file modification times
+				const bool bCopied = [FileManager copyItemAtPath:SrcDir.GetNSString() toPath:DestDir.GetNSString() error:nil];
+				if (!bCopied)
+				{
+					UE_LOG(LogFastBuildController, Warning, TEXT("Cannot use FASTBuild Shader Compiler as Metal shader compiler could not be copied to the intermediate folder: %s -> %s"),
+						*SrcDir, *DestDir);
+					FASTBuildControllerVariables::Enabled = 0;
+				}
+			}
+
+			bCopyMetalCompilerToIntermediateDir = false;
+		}
+
+		if (FASTBuildControllerVariables::Enabled == 1)
+		{
+			UE_LOG(LogFastBuildController, Warning, TEXT("FASTBuild Shader Compiler is temporarily disabled on Mac until problems with locating Autogen shader folder are sorted out"));
+			FASTBuildControllerVariables::Enabled = 0;
+		}
+
+#endif // PLATFORM_MAC
+	}
+
+	bSupported = FASTBuildControllerVariables::Enabled == 1;
+
+	return bSupported;
+}
+
+FString FFastBuildControllerModule::CreateUniqueFilePath()
+{
+	check(bSupported);
+	return FString::Printf(TEXT("%s/%d-fbuild"), *WorkingDirectory, NextFileID.Increment()); 
+}
+
+void FFastBuildControllerModule::StartupModule()
+{
+	check(!bModuleInitialized);
+
+	IModularFeatures::Get().RegisterModularFeature(GetModularFeatureType(), this);
+	
+	bModuleInitialized = true;
+}
+
+void FFastBuildControllerModule::ShutdownModule()
+{
+	check(bModuleInitialized);
+
+	IModularFeatures::Get().UnregisterModularFeature(GetModularFeatureType(), this);
+
+	// Stop the jobs thread
+	if (JobDispatcherThread)
+	{
+		JobDispatcherThread->Stop();	
+	}
+
+	// Cancel any remaining tasks
+	for (auto Iter = DispatchedTasks.CreateIterator(); Iter; ++Iter)
+	{
+		FDistributedBuildTaskResult Result;
+		Result.ReturnCode = 0;
+		Result.bCompleted = false;
+
+		FTask* Task = Iter.Value();
+		Task->Promise.SetValue(Result);
+		delete Task;
+	}
+
+	FTask* Task;
+	while (PendingTasks.Dequeue(Task))
+	{
+		FDistributedBuildTaskResult Result;
+		Result.ReturnCode = 0;
+		Result.bCompleted = false;
+		Task->Promise.SetValue(Result);
+		delete Task;
+	}
+
+	CleanWorkingDirectory();
+	bModuleInitialized = false;
+	bControllerInitialized = false;
+}
+
+void FFastBuildControllerModule::InitializeController()
+{
+	if (!IFileManager::Get().DirectoryExists(*RootWorkingDirectory))
+	{
+		IFileManager::Get().MakeDirectory(*RootWorkingDirectory);
+	}
+
+	if (IsSupported())
+	{
+		JobDispatcherThread = MakeUnique<FFastBuildJobProcessor>(*this);
+		JobDispatcherThread->StartThread();
+	}
+
+	bControllerInitialized = true;
+}
+
+void FFastBuildControllerModule::CleanWorkingDirectory()
+{
+	IFileManager& FileManager = IFileManager::Get();
+	
+	if (!WorkingDirectory.IsEmpty())
+	{
+		if (!FileManager.DeleteDirectory(*WorkingDirectory, false, true))
+		{
+			UE_LOG(LogFastBuildController, Log, TEXT("%s => Failed to delete current working Directory => %s"), ANSI_TO_TCHAR(__FUNCTION__), *RootWorkingDirectory);
+		}	
+	}
+}
+
+TFuture<FDistributedBuildTaskResult> FFastBuildControllerModule::EnqueueTask(const FTaskCommandData& CommandData)
+{
+	check(bSupported);
+
+	TPromise<FDistributedBuildTaskResult> Promise;
+	TFuture<FDistributedBuildTaskResult> Future = Promise.GetFuture();
+
+	// Enqueue the new task
+	FTask* Task = new FTask(NextTaskID.Increment(), CommandData, MoveTemp(Promise));
+	{
+		FScopeLock Lock(TasksCS.Get());
+		PendingTasks.Enqueue(Task);
+		PendingTasksCounter.Increment();
+	}
+
+	return MoveTemp(Future);
+}
+
+void FFastBuildControllerModule::EnqueueTask(FTask* Task)
+{
+	FScopeLock Lock(TasksCS.Get());
+	{
+		PendingTasks.Enqueue(Task);
+	}
+	PendingTasksCounter.Increment();
+}
+
+FTask* FFastBuildControllerModule::DequeueTask()
+{
+	FTask* TaskPtr;
+	
+	FScopeLock Lock(TasksCS.Get());
+	{
+		PendingTasks.Dequeue(TaskPtr);
+	}
+	
+	PendingTasksCounter.Increment();
+	return TaskPtr;
+}
+
+void FFastBuildControllerModule::RegisterDispatchedTask(FTask* DispatchedTask)
+{
+	check(DispatchedTask)
+	{
+		FScopeLock(TasksCS.Get());
+		DispatchedTasks.Add(DispatchedTask->ID , DispatchedTask);
+	}
+}
+
+void FFastBuildControllerModule::ReEnqueueDispatchedTasks()
+{
+	FScopeLock(TasksCS.Get());
+	// Reclaim dispatched (incomplete) tasks 
+	for (const TPair<uint32, FTask*>& DispatchedTaskEntry : DispatchedTasks)
+	{
+		EnqueueTask(DispatchedTaskEntry.Value);
+	}
+			
+	DispatchedTasks.Reset();
+}
+
+void FFastBuildControllerModule::DeRegisterDispatchedTasks(const TArray<uint32>& InTasksID)
+{
+	for (const uint32& TaskID : InTasksID)
+	{
+		FScopeLock(TasksCS.Get());
+		DispatchedTasks.Remove(TaskID);
+	}
+}
+
+void FFastBuildControllerModule::ReportJobProcessed(FTask* CompletedTask, const FTaskResponse& InTaskResponse)
+{
+	FDistributedBuildTaskResult Result;
+	Result.ReturnCode = InTaskResponse.ReturnCode;
+	Result.bCompleted = true;
+	CompletedTask->Promise.SetValue(Result);
+	delete CompletedTask;
+}
+
+FFastBuildControllerModule& FFastBuildControllerModule::Get()
+{
+	return FModuleManager::LoadModuleChecked<FFastBuildControllerModule>(TEXT("FastBuildController"));
+}
+
+IMPLEMENT_MODULE(FFastBuildControllerModule, FastBuildController);
diff --git a/Engine/Plugins/FastBuildController/Source/Private/FastBuildJobProcessor.cpp b/Engine/Plugins/FastBuildController/Source/Private/FastBuildJobProcessor.cpp
new file mode 100644
--- /dev/null
+++ b/Engine/Plugins/FastBuildController/Source/Private/FastBuildJobProcessor.cpp
@@ -0,0 +1,319 @@
+﻿// Copyright Epic Games, Inc. All Rights Reserved.
+
+#include "FastBuildJobProcessor.h"
+#include "FastBuildControllerModule.h"
+#include "FastBuildUtilities.h"
+#include "HAL/PlatformFileManager.h"
+
+namespace FastBuildJobProcessorOptions
+{
+	float SleepTimeBetweenActions = 0.1f;
+	FAutoConsoleVariableRef CVarFASTBuildSleepTimeBetweenActions(
+        TEXT("r.FASTBuild.JobProcessor.SleepTimeBetweenActions"),
+        SleepTimeBetweenActions,
+        TEXT("How much time the job processor thread should sleep between actions .\n"));
+
+
+	int32 MinJobsToKickOff = 100;
+	FAutoConsoleVariableRef CVarFASTBuildShaderMinJobsToKickOff(
+        TEXT("r.FASTBuild.JobProcessor.MinBatchSize"),
+        MinJobsToKickOff,
+        TEXT("Minimum number of shaders to compile with FASTBuild.\n")
+        TEXT("Default = 100\n"),
+        ECVF_Default);
+
+	int32 MaxTimeWithPendingJobs = 10;
+	FAutoConsoleVariableRef CVarFASTBuildShaderMaxTimeWithPendingJobs(
+        TEXT("r.FASTBuild.JobProcessor.MaxTimeWithPendingJobs"),
+        MaxTimeWithPendingJobs,
+        TEXT("Specifies how much time in seconds we will wait to have the min amount of pending jobs. Past this time, the build will start anyways.\n")
+        TEXT("Default = 10\n"),
+        ECVF_Default);
+}
+
+
+DEFINE_LOG_CATEGORY_STATIC(LogFastBuildJobProcessor, Log, Log);
+
+static const FString FASTBuildScriptFileName(TEXT("shaders.bff"));
+
+FFastBuildJobProcessor::FFastBuildJobProcessor(FFastBuildControllerModule& InControllerModule)
+	: Thread(nullptr),
+	  BuildProcessID(0),
+	  LastTimeKickedOffJobs(0),
+	  ControllerModule(InControllerModule),
+	  bIsWorkDone(false)
+{
+}
+
+FFastBuildJobProcessor::~FFastBuildJobProcessor()
+{
+	if (BuildProcessHandle.IsValid())
+	{
+		// We still have a build in progress, so we need to terminate it.
+		FPlatformProcess::TerminateProc(BuildProcessHandle);
+		FPlatformProcess::CloseProc(BuildProcessHandle);
+		FPlatformProcess::ClosePipe(PipeRead, PipeWrite);
+	}
+
+	if (Thread)
+	{
+		delete Thread;
+		Thread = nullptr;
+	}
+}
+
+uint32 FFastBuildJobProcessor::Run()
+{
+	bIsWorkDone = false;
+	LastTimeKickedOffJobs  = FPlatformTime::Cycles();
+	
+	while (!bForceStop)
+	{
+		const float ElapsedSeconds = (FPlatformTime::Cycles() - LastTimeKickedOffJobs) * FPlatformTime::GetSecondsPerCycle();
+		const bool bShouldKickoffJobsAnyway = ElapsedSeconds > FastBuildJobProcessorOptions::MaxTimeWithPendingJobs && ControllerModule.GetPendingTasksAmount() > 0;
+
+		
+		if (bShouldKickoffJobsAnyway || ControllerModule.GetPendingTasksAmount() > FastBuildJobProcessorOptions::MinJobsToKickOff)
+		{
+			LastTimeKickedOffJobs = FPlatformTime::Cycles();
+			SubmitPendingJobs();
+		}
+	
+		if (ControllerModule.AreTasksDispatched())
+		{
+			MonitorFastBuildProcess();
+			GatherBuildResults();
+		}	
+
+		FPlatformProcess::Sleep(FastBuildJobProcessorOptions::SleepTimeBetweenActions);
+	}
+
+	bIsWorkDone = true;
+	return 0;
+}
+
+void FFastBuildJobProcessor::StartThread()
+{
+	Thread = FRunnableThread::Create(this, TEXT("FastBuildJobProcessor"), 0, TPri_Normal, FPlatformAffinity::GetPoolThreadMask());
+}
+
+void FFastBuildJobProcessor::MonitorFastBuildProcess()
+{
+	if (!BuildProcessHandle.IsValid())
+	{
+		return;
+	}
+	
+	bool bDoExitCheck = false;
+	if (FPlatformProcess::IsProcRunning(BuildProcessHandle))
+	{
+		const FString STDOutput = FPlatformProcess::ReadPipe(PipeRead);
+		if (STDOutput.Len() > 0)
+		{
+			TArray<FString> Lines;
+			STDOutput.ParseIntoArrayLines(Lines);
+			for (const FString& Line : Lines)
+			{
+				UE_LOG(LogFastBuildJobProcessor, Verbose, TEXT("%s"), *Line);
+			}
+		}
+
+		if (!ControllerModule.AreTasksDispatchedOrPending())
+		{
+			// We've processed all batches.
+			// Wait for the FASTBuild console process to exit
+			FPlatformProcess::WaitForProc(BuildProcessHandle);
+			bDoExitCheck = true;
+		}
+	}
+	else
+	{
+		bDoExitCheck = true;
+	}
+	
+	if (bDoExitCheck)
+	{	
+		// The build process has stopped.
+		// Do one final pass over the output files to gather any remaining results.
+		GatherBuildResults();
+
+		// The build process is no longer running.
+		// We need to check the return code for possible failure
+		int32 ReturnCode = 0;
+		FPlatformProcess::GetProcReturnCode(BuildProcessHandle, &ReturnCode);
+
+		switch (ReturnCode)
+		{
+			case FBUILD_OK:
+				// No error
+				break;
+
+			case FBUILD_ERROR_LOADING_BFF:
+				UE_LOG(LogFastBuildJobProcessor, Fatal, TEXT("There was an error parsing FASTBuild script. This might be due to platform dependencies having files with the same names (Code %d)."), ReturnCode);
+				break;
+			case FBUILD_BUILD_FAILED:
+	        case FBUILD_BAD_ARGS:
+	        case FBUILD_FAILED_TO_SPAWN_WRAPPER:
+	        case FBUILD_FAILED_TO_SPAWN_WRAPPER_FINAL:
+	        case FBUILD_WRAPPER_CRASHED:
+	            // One or more of the shader compile worker processes crashed.
+	            UE_LOG(LogFastBuildJobProcessor, Fatal, TEXT("An error occurred during an FASTBuild shader compilation job. One or more of the shader compile worker processes exited unexpectedly (Code %d)."), ReturnCode);
+				break;
+			default:
+				UE_LOG(LogFastBuildJobProcessor, Display, TEXT("An unknown error occurred during an FASTBuild shader compilation job (Code %d). Incomplete shader jobs will be redispatched in another FASTBuild build."), ReturnCode);
+				break;
+			case FBUILD_ALREADY_RUNNING:
+				UE_LOG(LogFastBuildJobProcessor, Display, TEXT("FASTBuild is already running. Incomplete shader jobs will be redispatched in another FASTBuild build."));
+				break;
+		}
+
+		if (ReturnCode != FBUILD_OK && ReturnCode != FBUILD_ALREADY_RUNNING)
+		{
+			ControllerModule.ReEnqueueDispatchedTasks();
+		}
+
+		BuildProcessHandle.Reset();
+	}
+}
+
+void FFastBuildJobProcessor::WriteScriptFileToDisk(const TArray<FTask*>& PendingTasks, const FString& ScriptFilename, const FString& WorkerName) const
+{
+	// Create the FASTBuild script file.
+	TUniquePtr<FArchive> ScriptFile(FastBuildUtilities::CreateFileHelper(ScriptFilename));
+	check(ScriptFile.IsValid());
+	{
+		FastBuildUtilities::FASTBuildWriteScriptFileHeader(PendingTasks, *ScriptFile, WorkerName);
+		
+		// Write the task line for each shader batch
+		for (const FTask* CompilationTask : PendingTasks)
+		{
+			FString WorkerAbsoluteDirectory = IFileManager::Get().ConvertToAbsolutePathForExternalAppForWrite(*ControllerModule.GetWorkingDirectory());
+			FPaths::NormalizeDirectoryName(WorkerAbsoluteDirectory);
+
+			const FString ExecFunction = FString::Printf(
+                TEXT("ObjectList('ShaderBatch-%d')" LINE_TERMINATOR_ANSI)
+                TEXT("{" LINE_TERMINATOR_ANSI)
+                TEXT("\t.Compiler = 'ShaderCompiler'" LINE_TERMINATOR_ANSI)
+                TEXT("\t.CompilerOptions = '\"\" %d %d \"%%1\" \"%%2\"'" LINE_TERMINATOR_ANSI)
+                TEXT("\t.CompilerOutputExtension = '.out'"  LINE_TERMINATOR_ANSI)
+                TEXT("\t.CompilerInputFiles = { '%s' }" LINE_TERMINATOR_ANSI)
+                TEXT("\t.CompilerOutputPath = '%s'" LINE_TERMINATOR_ANSI)
+                TEXT("}"  LINE_TERMINATOR_ANSI LINE_TERMINATOR_ANSI),
+                CompilationTask->ID,
+                FGenericPlatformProcess::GetCurrentProcessId(),
+                CompilationTask->ID,
+                *CompilationTask->CommandData.InputFileName,
+                *WorkerAbsoluteDirectory);
+
+			ScriptFile->Serialize((void*)StringCast<ANSICHAR>(*ExecFunction, ExecFunction.Len()).Get(), sizeof(ANSICHAR) * ExecFunction.Len());
+		}
+
+		const FString AliasBuildTargetOpen = FString(
+            TEXT("Alias('all')" LINE_TERMINATOR_ANSI)
+            TEXT("{" LINE_TERMINATOR_ANSI)
+            TEXT("\t.Targets = { " LINE_TERMINATOR_ANSI)
+        );
+		ScriptFile->Serialize((void*)StringCast<ANSICHAR>(*AliasBuildTargetOpen, AliasBuildTargetOpen.Len()).Get(), sizeof(ANSICHAR) * AliasBuildTargetOpen.Len());
+
+		for (FTask* CompilationTask : PendingTasks)
+		{
+			const FString TargetExport = FString::Printf(TEXT("'ShaderBatch-%d', "), CompilationTask->ID);
+			ScriptFile->Serialize((void*)StringCast<ANSICHAR>(*TargetExport, TargetExport.Len()).Get(), sizeof(ANSICHAR) * TargetExport.Len());
+
+			ControllerModule.RegisterDispatchedTask(CompilationTask);
+		}
+	}
+
+	const FString AliasBuildTargetClose = FString(TEXT(" }"  LINE_TERMINATOR_ANSI "}"  LINE_TERMINATOR_ANSI));
+	ScriptFile->Serialize((void*)StringCast<ANSICHAR>(*AliasBuildTargetClose, AliasBuildTargetClose.Len()).Get(), sizeof(ANSICHAR) * AliasBuildTargetClose.Len());
+	ScriptFile = nullptr;
+}
+
+void FFastBuildJobProcessor::SubmitPendingJobs()
+{
+	if (BuildProcessHandle.IsValid() && FPlatformProcess::IsProcRunning(BuildProcessHandle))
+	{
+		return;
+	}
+
+	if (!ControllerModule.AreTasksPending())
+	{
+		return;
+	}
+
+	TArray<FTask*> TasksToSubmit;
+	FTask* PendingTask = nullptr;
+	{
+		FScopeLock(ControllerModule.GetTasksCS());
+		while (ControllerModule.AreTasksPending())
+		{
+			PendingTask = ControllerModule.DequeueTask();
+			TasksToSubmit.Add(PendingTask);
+		}
+	}
+	
+	const FString ScriptFilename = ControllerModule.GetWorkingDirectory() / FASTBuildScriptFileName;
+	const FString WorkerName = TasksToSubmit[0]->CommandData.Command;
+	const uint32 DispatcherID = TasksToSubmit[0]->CommandData.DispatcherPID;
+
+	// Create the Fast build script file with all the current pending jobs
+	WriteScriptFileToDisk(TasksToSubmit, ScriptFilename, WorkerName);
+
+	const FString FASTBuildConsoleArgs = TEXT("-config \"") + ScriptFilename + TEXT("\" -dist -clean -monitor");
+	
+	// Kick off the FASTBuild process...
+	verify(FPlatformProcess::CreatePipe(PipeRead, PipeWrite));
+	BuildProcessHandle = FPlatformProcess::CreateProc(*FastBuildUtilities::GetFastBuildExecutablePath(), *FASTBuildConsoleArgs, false, false, true, &BuildProcessID, 0, nullptr, PipeWrite);
+	if (!BuildProcessHandle.IsValid())
+	{
+		UE_LOG(LogFastBuildJobProcessor, Fatal, TEXT("Failed to launch %s during shader compilation."), *FastBuildUtilities::GetFastBuildExecutablePath());
+	}
+
+	// If the engine crashes, we don't get a chance to kill the build process.
+	// Start up the build monitor process to monitor for engine crashes.
+	if (WorkerName.Contains(TEXT("ShaderCompile")))
+	{
+		uint32 BuildMonitorProcessID;
+		FProcHandle BuildMonitorHandle = FPlatformProcess::CreateProc(*WorkerName, *FString::Printf(TEXT("-xgemonitor %d %d"), DispatcherID, BuildProcessID), true, false, false, &BuildMonitorProcessID, 0, nullptr, nullptr);
+		FPlatformProcess::CloseProc(BuildMonitorHandle);
+	}
+}
+
+void FFastBuildJobProcessor::GatherBuildResults() const
+{
+	IPlatformFile& PlatformFile = FPlatformFileManager::Get().GetPlatformFile();
+	IFileManager& FileManager = IFileManager::Get();
+	
+	TArray<uint32> CompletedTasks;
+	for (const TPair<uint32, FTask*>& DispatchedTaskEntry : ControllerModule.GetDispatchedTasks())
+	{
+		FTask* CompileTask = DispatchedTaskEntry.Value;
+
+		constexpr uint64 VersionAndFileSizeSize = sizeof(uint32) + sizeof(uint64);	
+		if (PlatformFile.FileExists(*CompileTask->CommandData.OutputFileName) &&
+			FileManager.FileSize(*CompileTask->CommandData.OutputFileName) > VersionAndFileSizeSize)
+		{
+			TUniquePtr<FArchive> OutputFilePtr(FileManager.CreateFileReader(*CompileTask->CommandData.OutputFileName, FILEREAD_Silent));
+			if (OutputFilePtr)
+			{
+				FArchive& OutputFile = *OutputFilePtr;
+				int32 OutputVersion;
+				OutputFile << OutputVersion; // NOTE (SB): Do not care right now about the version.
+				int64 FileSize = 0;
+				OutputFile << FileSize;
+
+				// NOTE (SB): Check if we received the full file yet.
+				if (OutputFile.TotalSize() >= FileSize)
+				{
+					FTaskResponse TaskCompleted;
+					TaskCompleted.ID = CompileTask->ID;
+					TaskCompleted.ReturnCode = 0;
+					
+					ControllerModule.ReportJobProcessed(CompileTask, TaskCompleted);
+					CompletedTasks.Add(DispatchedTaskEntry.Key);
+				}
+			}
+		}
+	}
+	ControllerModule.DeRegisterDispatchedTasks(CompletedTasks);
+}
diff --git a/Engine/Plugins/FastBuildController/Source/Private/FastBuildUtilities.cpp b/Engine/Plugins/FastBuildController/Source/Private/FastBuildUtilities.cpp
new file mode 100644
--- /dev/null
+++ b/Engine/Plugins/FastBuildController/Source/Private/FastBuildUtilities.cpp
@@ -0,0 +1,274 @@
+﻿// Copyright Epic Games, Inc. All Rights Reserved.
+
+#include "FastBuildUtilities.h"
+#include <DistributedBuildInterface/Public/DistributedBuildControllerInterface.h>
+
+#include "ShaderCore.h"
+#include "HAL/PlatformFileManager.h"
+#include "Interfaces/ITargetPlatform.h"
+#include "Interfaces/ITargetPlatformManagerModule.h"
+#include "HAL/PlatformMisc.h"
+
+namespace FASTBuildControllerUtilitiesVariables
+{
+	int32 SendWorkerApplicationDebugSymbols = 0;
+	FAutoConsoleVariableRef CVarFASTBuildSendDebugSymbols(
+        TEXT("r.FASTBuildController.SendSCWDebugSymbols"),
+		SendWorkerApplicationDebugSymbols,
+        TEXT("Enable when distributed shader compiler workers crash.\n")
+        TEXT("0: Do not send along debug information in FASTBuild. \n")
+        TEXT("1: Send along debug information in FASTBuild."),
+        ECVF_Default);
+
+	int32 SendAllPossibleShaderDependencies = 0;
+	FAutoConsoleVariableRef CVarFASTBuildSendAllPossibleShaderDependencies(
+        TEXT("r.FASTBuildController.SendAllPossibleShaderDependencies"),
+        SendAllPossibleShaderDependencies,
+        TEXT("Send all possible dependencies of the shaders to the remote machines.")
+        TEXT("0: Use dependencies array reported in the task structure.\n")
+        TEXT("1: Brute-force discover all possible dependencies. \n"),
+        ECVF_Default);
+}
+
+FString FastBuildUtilities::GetShaderFullPathOfShaderDependency(const FString& InVirtualPath)
+{
+	FString DependencyFilename = GetShaderSourceFilePath(InVirtualPath);
+	DependencyFilename = IFileManager::Get().ConvertToAbsolutePathForExternalAppForWrite(*DependencyFilename);
+	FPaths::NormalizeDirectoryName(DependencyFilename);
+	return DependencyFilename;
+}
+
+void FastBuildUtilities::WriteShaderWorkerDependenciesToFile(FArchive& ScriptFile)
+{
+	FDependencyEnumerator DllDeps = FDependencyEnumerator(ScriptFile, TEXT("ShaderCompileWorker-"), PLATFORM_WINDOWS == 1 ? TEXT(".dll") : (PLATFORM_MAC == 1 ? TEXT(".dylib") : TEXT(".so")));
+	IFileManager::Get().IterateDirectoryRecursively(*FPlatformProcess::GetModulesDirectory(), DllDeps);
+	FDependencyEnumerator ModulesDeps = FDependencyEnumerator(ScriptFile, TEXT("ShaderCompileWorker"), TEXT(".modules"));
+	IFileManager::Get().IterateDirectoryRecursively(*FPlatformProcess::GetModulesDirectory(), ModulesDeps);
+	if (FASTBuildControllerUtilitiesVariables::SendWorkerApplicationDebugSymbols)
+	{
+#if PLATFORM_WINDOWS
+		FDependencyEnumerator PdbDeps = FDependencyEnumerator(ScriptFile, TEXT("ShaderCompileWorker"), TEXT(".pdb"));
+		IFileManager::Get().IterateDirectoryRecursively(*FPlatformProcess::GetModulesDirectory(), PdbDeps);
+#endif
+	}
+
+	FDependencyEnumerator IniDeps = FDependencyEnumerator(ScriptFile, nullptr, TEXT(".ini"));
+	TArray<FString> EngineConfigDirs = FPaths::GetExtensionDirs(FPaths::EngineDir(), TEXT("Config"));
+	for (const FString& ConfigDir : EngineConfigDirs)
+	{
+		IFileManager::Get().IterateDirectoryRecursively(*ConfigDir, IniDeps);
+	}
+}
+
+FString FastBuildUtilities::ReplaceEnvironmentVariablesInPath(const FString& ExtraFilePartialPath)
+{
+	FString ParsedPath;
+
+	// Fast build cannot read environmental variables easily
+	// Is better to resolve them here
+	if (ExtraFilePartialPath.Contains(TEXT("%")))
+	{
+		TArray<FString> PathSections;
+		ExtraFilePartialPath.ParseIntoArray(PathSections,TEXT("/"));
+
+		for (FString& Section : PathSections)
+		{
+			if (Section.Contains(TEXT("%")))
+			{
+				Section.RemoveFromStart(TEXT("%"));
+				Section.RemoveFromEnd(TEXT("%"));
+				Section = FPlatformMisc::GetEnvironmentVariable(*Section);
+			}
+		}
+
+		for (FString& Section : PathSections)
+		{
+			ParsedPath /= Section;
+		}
+			
+		FPaths::NormalizeDirectoryName(ParsedPath);
+	}
+
+	if (ParsedPath.IsEmpty())
+	{
+		ParsedPath = ExtraFilePartialPath;
+	}
+	
+	return ParsedPath;
+}
+
+void FastBuildUtilities::WriteDependenciesToFileUsingWildcardPath(FArchive& ScriptFile, FString ParsedPath)
+{
+	if (ParsedPath.Contains(TEXT("*")))
+	{
+		TArray<FString> PathSections;
+		ParsedPath.ParseIntoArray(PathSections,TEXT("/"));
+
+		FString PartialPath;
+		for (const FString& Section : PathSections)
+		{
+			if (Section.Contains(TEXT("*")))
+			{
+				FString Prefix;
+				FString Extension;
+				Section.Split(TEXT("*"),&Prefix,&Extension);
+		    			
+				FDependencyEnumerator WildcardQuery = FDependencyEnumerator(ScriptFile, *Prefix, *Extension);
+				IFileManager::Get().IterateDirectoryRecursively(*PartialPath, WildcardQuery);
+				break;
+			}
+			PartialPath /= Section;
+		}
+	}
+}
+
+void FastBuildUtilities::WritePlatformCompilerDependenciesToFile(FArchive& ScriptFile)
+{
+	ITargetPlatformManagerModule* TargetPlatformManager = GetTargetPlatformManager();
+	TArray<FString> FASTBuild_Toolchain;
+	for (ITargetPlatform* TargetPlatform : TargetPlatformManager->GetTargetPlatforms())
+	{
+		TargetPlatform->GetShaderCompilerDependencies(FASTBuild_Toolchain);
+	}
+
+	for (const FString& ExtraFilePartialPath : FASTBuild_Toolchain)
+	{
+		FString ParsedPath = ReplaceEnvironmentVariablesInPath(ExtraFilePartialPath);
+
+		if (ParsedPath.Contains(TEXT("*")))
+		{
+			// Fast Build cannot resolve wildcards
+		    // We need to resolve them here
+		    WriteDependenciesToFileUsingWildcardPath(ScriptFile, ParsedPath);
+		}
+		else
+		{
+			ParsedPath = TEXT("\t\t'") + IFileManager::Get().ConvertToAbsolutePathForExternalAppForRead(*(ParsedPath)) + TEXT("'," LINE_TERMINATOR_ANSI);
+			ScriptFile.Serialize((void*)StringCast<ANSICHAR>(*ParsedPath, ParsedPath.Len()).Get(), sizeof(ANSICHAR) * ParsedPath.Len());
+		}
+	
+	}
+}
+
+void FastBuildUtilities::WriteDependenciesForShaderToScript(const TArray<FTask*>& InCompilationTasks, FArchive& ScriptFile)
+{
+	if (FASTBuildControllerUtilitiesVariables::SendAllPossibleShaderDependencies)
+	{
+		// This is kinda a hack because we are sending all possible dependencies
+		// Warning: this can result in sending files with the same name located in different paths. FASTBuild currently (as of v1.05) does not handle this, and BFF reading will be broken.
+		FDependencyEnumerator ShaderUsfDeps = FDependencyEnumerator(ScriptFile, nullptr, TEXT(".usf"));
+		FDependencyEnumerator ShaderUshDeps = FDependencyEnumerator(ScriptFile, nullptr, TEXT(".ush"));
+		FDependencyEnumerator ShaderHeaderDeps = FDependencyEnumerator(ScriptFile, nullptr, TEXT(".h"));
+		const TMap<FString, FString> ShaderSourceDirectoryMappings = AllShaderSourceDirectoryMappings();
+		for (auto& ShaderDirectoryMapping : ShaderSourceDirectoryMappings)
+		{
+			IFileManager::Get().IterateDirectoryRecursively(*ShaderDirectoryMapping.Value, ShaderUsfDeps);
+			IFileManager::Get().IterateDirectoryRecursively(*ShaderDirectoryMapping.Value, ShaderUshDeps);
+			IFileManager::Get().IterateDirectoryRecursively(*ShaderDirectoryMapping.Value, ShaderHeaderDeps);
+		}
+	}
+	else
+	{
+		TArray<FString> FullUniqueDependenciesArray;
+		{
+			for (FTask* CompilationTask : InCompilationTasks)
+			{
+				for (const FString& Dependency : CompilationTask->CommandData.Dependencies)
+				{
+					FullUniqueDependenciesArray.AddUnique(Dependency);
+				}
+			}
+		}
+
+		for (const FString& ExtraDependency : FullUniqueDependenciesArray)
+		{
+			const FString ExtraFile = TEXT("\t\t'") + GetShaderFullPathOfShaderDependency(ExtraDependency) + TEXT("'," LINE_TERMINATOR_ANSI);
+			ScriptFile.Serialize((void*)StringCast<ANSICHAR>(*ExtraFile, ExtraFile.Len()).Get(), sizeof(ANSICHAR) * ExtraFile.Len());
+		}
+	}
+
+#if PLATFORM_MAC
+	const FString MetalIntermediateDir = FPaths::EngineIntermediateDir() + TEXT("/Shaders/metal");
+	FDependencyEnumerator MetalCompilerDeps = FDependencyEnumerator(ScriptFile, nullptr, nullptr);
+	IFileManager::Get().IterateDirectoryRecursively(*MetalIntermediateDir, MetalCompilerDeps);
+#endif
+}
+
+void FastBuildUtilities::FASTBuildWriteScriptFileHeader(const TArray<FTask*>& InCompilationTasks, FArchive& ScriptFile, const FString& WorkerName)
+{	
+	static const TCHAR HeaderTemplate[] =
+		TEXT("Settings" LINE_TERMINATOR_ANSI)
+		TEXT("{" LINE_TERMINATOR_ANSI)
+		TEXT("\t.CachePath = '%s'" LINE_TERMINATOR_ANSI)
+		TEXT("}" LINE_TERMINATOR_ANSI)
+		TEXT(LINE_TERMINATOR_ANSI)
+		TEXT("Compiler('ShaderCompiler')" LINE_TERMINATOR_ANSI)
+		TEXT("{" LINE_TERMINATOR_ANSI)
+		TEXT("\t.CompilerFamily = 'custom'" LINE_TERMINATOR_ANSI)
+		TEXT("\t.Executable = '%s'" LINE_TERMINATOR_ANSI)
+		TEXT("\t.ExecutableRootPath = '%s'" LINE_TERMINATOR_ANSI)
+		TEXT("\t.SimpleDistributionMode = true" LINE_TERMINATOR_ANSI)
+		TEXT("\t.ExtraFiles = " LINE_TERMINATOR_ANSI)
+		TEXT("\t{" LINE_TERMINATOR_ANSI);
+
+	const FString HeaderString = FString::Printf(HeaderTemplate, *FastBuildUtilities::GetFastBuildCache(), *WorkerName,
+			*IFileManager::Get().ConvertToAbsolutePathForExternalAppForRead(*FPaths::RootDir()));
+	ScriptFile.Serialize((void*)StringCast<ANSICHAR>(*HeaderString, HeaderString.Len()).Get(),
+			sizeof(ANSICHAR) * HeaderString.Len());
+
+	WritePlatformCompilerDependenciesToFile(ScriptFile);
+
+	WriteShaderWorkerDependenciesToFile(ScriptFile);
+
+	WriteDependenciesForShaderToScript(InCompilationTasks, ScriptFile);
+
+	const FString ExtraFilesFooter =
+		TEXT("\t}" LINE_TERMINATOR_ANSI)
+		TEXT("}" LINE_TERMINATOR_ANSI);
+	ScriptFile.Serialize((void*)StringCast<ANSICHAR>(*ExtraFilesFooter, ExtraFilesFooter.Len()).Get(), sizeof(ANSICHAR) * ExtraFilesFooter.Len());
+}
+
+FArchive* FastBuildUtilities::CreateFileHelper(const FString& InFileName)
+{
+	// TODO: This logic came from FShaderCompileThreadRunnable::WriteNewTasks().
+	// We can't avoid code duplication unless we refactored the local worker too.
+
+	FArchive* File = nullptr;
+	int32 RetryCount = 0;
+	// Retry over the next two seconds if we can't write out the file.
+	// Anti-virus and indexing applications can interfere and cause this to fail.
+	while (File == nullptr && RetryCount < 200)
+	{
+		if (RetryCount > 0)
+		{
+			FPlatformProcess::Sleep(0.01f);
+		}
+		File = IFileManager::Get().CreateFileWriter(*InFileName, FILEWRITE_EvenIfReadOnly);
+		RetryCount++;
+	}
+	if (File == nullptr)
+	{
+		File = IFileManager::Get().CreateFileWriter(*InFileName, FILEWRITE_EvenIfReadOnly | FILEWRITE_NoFail);
+	}
+	checkf(File, TEXT("Failed to create file %s!"), *InFileName);
+	return File;
+}
+
+FString FastBuildUtilities::GetFastBuildExecutablePath()
+{
+	FString FASTBuildExecutablePath = TEXT("PLATFORM NOT SUPPORTED");
+#if PLATFORM_WINDOWS
+	FASTBuildExecutablePath = TEXT("Extras\\ThirdPartyNotUE\\FASTBuild\\Win64\\FBuild.exe");
+#elif PLATFORM_MAC
+	FASTBuildExecutablePath = TEXT("Extras/ThirdPartyNotUE/FASTBuild/Mac/FBuild");
+#elif PLATFORM_LINUX
+	FASTBuildExecutablePath = TEXT("Extras/ThirdPartyNotUE/FASTBuild/Linux/fbuild");
+#endif
+
+	return FPaths::EngineDir() / FASTBuildExecutablePath;
+}
+
+FString FastBuildUtilities::GetFastBuildCache()
+{
+	return FPaths::ProjectSavedDir() / TEXT("FASTBuildCache");
+}
diff --git a/Engine/Plugins/FastBuildController/Source/Public/FastBuildControllerModule.h b/Engine/Plugins/FastBuildController/Source/Public/FastBuildControllerModule.h
new file mode 100644
--- /dev/null
+++ b/Engine/Plugins/FastBuildController/Source/Public/FastBuildControllerModule.h
@@ -0,0 +1,101 @@
+// Copyright Epic Games, Inc. All Rights Reserved.
+
+#pragma once
+
+#include "DistributedBuildInterface/Public/DistributedBuildControllerInterface.h"
+#include "Containers/Array.h"
+
+class FFastBuildJobProcessor;
+
+class FFastBuildControllerModule : public IDistributedBuildController
+{
+	bool bSupported;
+	bool bModuleInitialized;
+	bool bControllerInitialized;
+
+	FThreadSafeCounter NextFileID;
+	FThreadSafeCounter NextTaskID;
+
+	FThreadSafeCounter PendingTasksCounter;
+
+	/** The thread spawned for shader compiling. */
+	TUniquePtr<FFastBuildJobProcessor> JobDispatcherThread;
+	
+	// Taken when accessing the PendingTasks and DispatchedTasks members.
+	TSharedPtr<FCriticalSection> TasksCS;
+
+	// Queue of tasks submitted by the engine, but not yet dispatched to the controller.
+	TQueue<FTask*> PendingTasks;
+
+	// Map of tasks dispatched to the controller and running within FastBuild, that have not yet finished.
+	TMap<uint32,FTask*> DispatchedTasks;
+
+	const FString RootWorkingDirectory;
+
+	const FString WorkingDirectory;
+
+	bool bShutdown;
+
+	void CleanWorkingDirectory();
+
+public:
+	FFastBuildControllerModule();
+	virtual ~FFastBuildControllerModule();
+
+	virtual void StartupModule() override final;
+	virtual void ShutdownModule() override final;
+
+	virtual bool RequiresRelativePaths() override final{ return true; };
+
+	virtual void InitializeController() override final;
+
+	virtual const FString GetName() override final { return FString("FastBuild Controller"); };
+
+	virtual bool IsSupported() override final;
+
+	virtual FString CreateUniqueFilePath() override final;
+	virtual TFuture<FDistributedBuildTaskResult> EnqueueTask(const FTaskCommandData& CommandData) override final;
+	void EnqueueTask(FTask* Task);
+	FTask* DequeueTask();
+
+	void RegisterDispatchedTask(FTask* DispatchedTask);
+	void ReEnqueueDispatchedTasks();
+	void DeRegisterDispatchedTasks(const TArray<uint32>& InTasksID);
+	void ReportJobProcessed(FTask* CompletedTask, const FTaskResponse& InTaskResponse);
+	
+	static FFastBuildControllerModule& Get();
+
+	inline bool AreTasksPending() const
+	{
+		FScopeLock Lock(TasksCS.Get());
+		return  !PendingTasks.IsEmpty();
+	}
+
+	inline bool AreTasksDispatched() const
+	{
+		FScopeLock Lock(TasksCS.Get());
+		return DispatchedTasks.Num() > 0;
+	}
+
+	inline bool AreTasksDispatchedOrPending() const
+	{
+		FScopeLock Lock(TasksCS.Get());
+		return DispatchedTasks.Num() > 0 || !PendingTasks.IsEmpty();
+	}
+
+	int32 GetPendingTasksAmount() const { return PendingTasksCounter.GetValue();}
+
+	FString GetRootWorkingDirectory() const { return RootWorkingDirectory;}
+
+	FString GetWorkingDirectory() const { return WorkingDirectory;}
+
+	const TMap<uint32, FTask*>& GetDispatchedTasks() const
+	{
+		return DispatchedTasks;
+	}
+
+	FCriticalSection* GetTasksCS() const
+	{
+		return TasksCS.Get();
+	}
+};
\ No newline at end of file
diff --git a/Engine/Plugins/FastBuildController/Source/Public/FastBuildJobProcessor.h b/Engine/Plugins/FastBuildController/Source/Public/FastBuildJobProcessor.h
new file mode 100644
--- /dev/null
+++ b/Engine/Plugins/FastBuildController/Source/Public/FastBuildJobProcessor.h
@@ -0,0 +1,72 @@
+﻿// Copyright Epic Games, Inc. All Rights Reserved.
+
+#pragma once
+
+#include "HAL/Runnable.h"
+#include "HAL/ThreadSafeBool.h"
+
+class FFastBuildControllerModule;
+class FShaderBatch;
+struct FTask;
+
+enum EFASTBuild_ReturnCodes
+{
+	FBUILD_OK = 0,
+	FBUILD_BUILD_FAILED = -1,
+	FBUILD_ERROR_LOADING_BFF = -2,
+	FBUILD_BAD_ARGS = -3,
+	FBUILD_ALREADY_RUNNING = -4,
+	FBUILD_FAILED_TO_SPAWN_WRAPPER = -5,
+	FBUILD_FAILED_TO_SPAWN_WRAPPER_FINAL = -6,
+	FBUILD_WRAPPER_CRASHED = -7,
+};
+class FFastBuildJobProcessor : public FRunnable
+{
+public:
+
+	FFastBuildJobProcessor(FFastBuildControllerModule& InControllerModule);
+	virtual ~FFastBuildJobProcessor();
+
+	/** Main loop */
+	virtual uint32 Run() override;
+	
+	/** Aborts the main loop as soon as possible */
+	virtual void Stop() override { bForceStop = true; };
+	
+	/** Creates the threads and starts the main loop */
+	void StartThread();
+
+	/** Used to know when this thread has finished the main loop */
+	inline bool IsWorkDone() const {return bIsWorkDone;};
+
+protected:
+
+	/** The runnable thread */
+	FRunnableThread* Thread;
+	
+	FProcHandle BuildProcessHandle;
+	uint32 BuildProcessID;
+	void* PipeRead = nullptr;
+	void* PipeWrite = nullptr;
+
+	uint32 LastTimeKickedOffJobs;
+	
+	FFastBuildControllerModule& ControllerModule;
+	
+	/** Used to abort the current processing loop */
+	FThreadSafeBool bForceStop;
+	
+	/** Set to true when the main loop finishes*/
+	FThreadSafeBool bIsWorkDone;
+	
+	void WriteScriptFileToDisk(const TArray<FTask*>& PendingTasks, const FString& ScriptFilename, const FString& WorkerName) const;
+
+	/** Takes all tasks from the queue and creates a FastBuild Script to the kick off the pending jobs */
+	void SubmitPendingJobs();
+	
+	/** Checks the file system to find any results form the submitted jobs */
+	void GatherBuildResults() const;
+
+	/** Checks if the FastBuild process is still running properly and handles bad situations */
+	void MonitorFastBuildProcess();
+};
\ No newline at end of file
diff --git a/Engine/Plugins/FastBuildController/Source/Public/FastBuildUtilities.h b/Engine/Plugins/FastBuildController/Source/Public/FastBuildUtilities.h
new file mode 100644
--- /dev/null
+++ b/Engine/Plugins/FastBuildController/Source/Public/FastBuildUtilities.h
@@ -0,0 +1,56 @@
+﻿// Copyright Epic Games, Inc. All Rights Reserved.
+
+#pragma once
+
+struct FTask;
+
+class FDependencyEnumerator : public IPlatformFile::FDirectoryVisitor
+{
+public:
+	FDependencyEnumerator(FArchive& InScriptFile, const TCHAR* InPrefix, const TCHAR* InExtension)
+        : ScriptFile(InScriptFile)
+        , Prefix(InPrefix)
+        , Extension(InExtension)
+		
+	{
+	}
+
+	virtual bool Visit(const TCHAR* FilenameChar, bool bIsDirectory) override
+	{
+		if (!bIsDirectory)
+		{
+			const FString Filename = FString(FilenameChar);
+
+			if ((!Prefix || Filename.Contains(Prefix)) && (!Extension || Filename.EndsWith(Extension)))
+			{
+				const FString ExtraFile = TEXT("\t\t'") + IFileManager::Get().ConvertToAbsolutePathForExternalAppForWrite(*Filename) + TEXT("'," LINE_TERMINATOR_ANSI);
+				ScriptFile.Serialize((void*)StringCast<ANSICHAR>(*ExtraFile, ExtraFile.Len()).Get(), sizeof(ANSICHAR) * ExtraFile.Len());
+			}
+		}
+
+		return true;
+	}
+
+	FArchive& ScriptFile;
+	const TCHAR* Prefix;
+	const TCHAR* Extension;
+};
+
+class FastBuildUtilities
+{
+public:
+
+	static void FASTBuildWriteScriptFileHeader(const TArray<FTask*>& InCompilationTasks, FArchive& ScriptFile, const FString& WorkerName);
+	static FArchive* CreateFileHelper(const FString& InFileName);
+
+	static FString GetFastBuildExecutablePath();
+	static FString GetFastBuildCache();
+
+protected:
+	static FString GetShaderFullPathOfShaderDependency(const FString& InVirtualPath);
+	static void WriteShaderWorkerDependenciesToFile(FArchive& ScriptFile);
+	static FString ReplaceEnvironmentVariablesInPath(const FString& ExtraFilePartialPath);
+	static void WriteDependenciesToFileUsingWildcardPath(FArchive& ScriptFile, FString ParsedPath);
+	static void WritePlatformCompilerDependenciesToFile(FArchive& ScriptFile);
+	static void WriteDependenciesForShaderToScript(const TArray<FTask*>& InCompilationTasks, FArchive& ScriptFile);
+};
diff --git a/Engine/Plugins/XGEController/Source/Private/XGEControllerModule.cpp b/Engine/Plugins/XGEController/Source/Private/XGEControllerModule.cpp
--- a/Engine/Plugins/XGEController/Source/Private/XGEControllerModule.cpp
+++ b/Engine/Plugins/XGEController/Source/Private/XGEControllerModule.cpp
@@ -567,7 +567,7 @@ void FXGEControllerModule::ReadBackThreadProc()
 FString FXGEControllerModule::CreateUniqueFilePath()
 {
 	check(bSupported);
-	return FString::Printf(TEXT("%s/%d.xge"), *WorkingDirectory, NextFileID.Increment());
+	return FString::Printf(TEXT("%s/%d-xge"), *WorkingDirectory, NextFileID.Increment());
 }
 
 TFuture<FDistributedBuildTaskResult> FXGEControllerModule::EnqueueTask(const FTaskCommandData& CommandData)
diff --git a/Engine/Source/Developer/DistributedBuildInterface/Public/DistributedBuildControllerInterface.h b/Engine/Source/Developer/DistributedBuildInterface/Public/DistributedBuildControllerInterface.h
--- a/Engine/Source/Developer/DistributedBuildInterface/Public/DistributedBuildControllerInterface.h
+++ b/Engine/Source/Developer/DistributedBuildInterface/Public/DistributedBuildControllerInterface.h
@@ -17,6 +17,8 @@ struct FTaskCommandData
 	FString Command;
 	FString CommandArgs;
 	FString InputFileName;
+	FString OutputFileName;
+	uint32 DispatcherPID = 0;
 	TArray<FString> Dependencies;
 };
 
@@ -43,6 +45,8 @@ class IDistributedBuildController : public IModuleInterface, public IModularFeat
 {
 public:
 	virtual bool SupportsDynamicReloading() override { return false; }
+	
+	virtual bool RequiresRelativePaths() { return false; }
 
 	virtual void InitializeController() = 0;
 	
diff --git a/Engine/Source/Developer/Mac/MacTargetPlatform/Private/GenericMacTargetPlatform.h b/Engine/Source/Developer/Mac/MacTargetPlatform/Private/GenericMacTargetPlatform.h
--- a/Engine/Source/Developer/Mac/MacTargetPlatform/Private/GenericMacTargetPlatform.h
+++ b/Engine/Source/Developer/Mac/MacTargetPlatform/Private/GenericMacTargetPlatform.h
@@ -310,6 +310,16 @@ return TSuper::SupportsFeature(Feature);
 		TextureLODSettings = InTextureLODSettings;
 	}
 
+	virtual bool CanSupportRemoteShaderCompile() const override
+	{
+		return true;
+	}
+	
+	virtual void GetShaderCompilerDependencies(TArray<FString>& OutDependencies) const override
+	{
+		FTargetPlatformBase::AddDependencySCArrayHelper(OutDependencies, TEXT("Binaries/ThirdParty/ShaderConductor/Mac/libdxcompiler.dylib"));
+		FTargetPlatformBase::AddDependencySCArrayHelper(OutDependencies, TEXT("Binaries/ThirdParty/ShaderConductor/Mac/libShaderConductor.dylib"));
+	}
 
 	virtual FName GetWaveFormat( const class USoundWave* Wave ) const override
 	{
diff --git a/Engine/Source/Runtime/Engine/Private/ShaderCompiler/Experimental/ShaderCompileFASTBuild.cpp b/Engine/Source/Runtime/Engine/Private/ShaderCompiler/Experimental/ShaderCompileFASTBuild.cpp
deleted file mode 100644
diff --git a/Engine/Source/Runtime/Engine/Private/ShaderCompiler/ShaderCompiler.cpp b/Engine/Source/Runtime/Engine/Private/ShaderCompiler/ShaderCompiler.cpp
--- a/Engine/Source/Runtime/Engine/Private/ShaderCompiler/ShaderCompiler.cpp
+++ b/Engine/Source/Runtime/Engine/Private/ShaderCompiler/ShaderCompiler.cpp
@@ -3012,7 +3012,6 @@ FShaderCompilingManager::FShaderCompilingManager() :
 	NumShaderCompilingThreadsDuringGame = FMath::Min<int32>(NumShaderCompilingThreadsDuringGame, NumShaderCompilingThreads);
 
 	TUniquePtr<FShaderCompileThreadRunnableBase> RemoteCompileThread;
-#if PLATFORM_WINDOWS
 	const bool bCanUseRemoteCompiling = bAllowCompilingThroughWorkers && ShaderCompiler::IsRemoteCompilingAllowed() && AllTargetPlatformSupportsRemoteShaderCompiling();
 	BuildDistributionController = bCanUseRemoteCompiling ? FindRemoteCompilerController() : nullptr;
 	
@@ -3021,15 +3020,6 @@ FShaderCompilingManager::FShaderCompilingManager() :
 		UE_LOG(LogShaderCompilers, Display, TEXT("Using %s for Shader Compilation."), *BuildDistributionController->GetName());
 		RemoteCompileThread = MakeUnique<FShaderCompileDistributedThreadRunnable_Interface>(this, *BuildDistributionController);
 	}
-	else
-#endif // PLATFORM_WINDOWS
-#if PLATFORM_DESKTOP
-	if (bAllowCompilingThroughWorkers && ShaderCompiler::IsRemoteCompilingAllowed() && FShaderCompileFASTBuildThreadRunnable::IsSupported())
-	{
-		UE_LOG(LogShaderCompilers, Display, TEXT("Using FASTBuild Shader Compiler."));
-		RemoteCompileThread = MakeUnique<FShaderCompileFASTBuildThreadRunnable>(this);
-	}
-#endif // PLATFORM_DESKTOP
 
 	GConfig->SetBool(TEXT("/Script/UnrealEd.UnrealEdOptions"), TEXT("UsingXGE"), RemoteCompileThread.IsValid(), GEditorIni);
 
diff --git a/Engine/Source/Runtime/Engine/Private/ShaderCompiler/ShaderCompilerDistributed.cpp b/Engine/Source/Runtime/Engine/Private/ShaderCompiler/ShaderCompilerDistributed.cpp
--- a/Engine/Source/Runtime/Engine/Private/ShaderCompiler/ShaderCompilerDistributed.cpp
+++ b/Engine/Source/Runtime/Engine/Private/ShaderCompiler/ShaderCompilerDistributed.cpp
@@ -68,8 +68,9 @@ FShaderCompileDistributedThreadRunnable_Interface::~FShaderCompileDistributedThr
 
 void FShaderCompileDistributedThreadRunnable_Interface::DispatchShaderCompileJobsBatch(TArray<FShaderCommonCompileJobPtr>& JobsToSerialize)
 {
-	FString InputFilePath = CachedController.CreateUniqueFilePath();
-	FString OutputFilePath = CachedController.CreateUniqueFilePath();
+	const FString BaseFilePath = CachedController.CreateUniqueFilePath();
+	FString InputFilePath = BaseFilePath + TEXT(".in");
+	FString OutputFilePath = BaseFilePath + TEXT(".out");
 
 	const FString WorkingDirectory = FPaths::GetPath(InputFilePath);
 	const FString InputFileName = FPaths::GetCleanFilename(InputFilePath);
@@ -86,7 +87,7 @@ void FShaderCompileDistributedThreadRunnable_Interface::DispatchShaderCompileJob
 
 	// Serialize the jobs to the input file
 	FArchive* InputFileAr = IFileManager::Get().CreateFileWriter(*InputFilePath, FILEWRITE_EvenIfReadOnly | FILEWRITE_NoFail);
-	FShaderCompileUtilities::DoWriteTasks(JobsToSerialize, *InputFileAr);
+	FShaderCompileUtilities::DoWriteTasks(JobsToSerialize, *InputFileAr, CachedController.RequiresRelativePaths());
 	delete InputFileAr;
 
 	// Kick off the job
@@ -96,6 +97,8 @@ void FShaderCompileDistributedThreadRunnable_Interface::DispatchShaderCompileJob
 	TaskCommandData.Command = Manager->ShaderCompileWorkerName;
 	TaskCommandData.CommandArgs = WorkerParameters;
 	TaskCommandData.InputFileName = InputFilePath;
+	TaskCommandData.OutputFileName = OutputFilePath;
+	TaskCommandData.DispatcherPID = Manager->ProcessId;
 	TaskCommandData.Dependencies = GetDependencyFilesForJobs(JobsToSerialize);
 	
 	DispatchedTasks.Add(
diff --git a/Engine/Source/Runtime/Engine/Public/ShaderCompiler.h b/Engine/Source/Runtime/Engine/Public/ShaderCompiler.h
--- a/Engine/Source/Runtime/Engine/Public/ShaderCompiler.h
+++ b/Engine/Source/Runtime/Engine/Public/ShaderCompiler.h
@@ -685,100 +685,6 @@ private:
 	void DispatchShaderCompileJobsBatch(TArray<FShaderCommonCompileJobPtr>& JobsToSerialize);
 };
 
-class FShaderCompileFASTBuildThreadRunnable : public FShaderCompileThreadRunnableBase
-{
-private:
-	/** The handle referring to the XGE console process, if a build is in progress. */
-	FProcHandle BuildProcessHandle;
-	void* PipeRead = nullptr;
-	void* PipeWrite = nullptr;
-
-	/** Process ID of the XGE console, if a build is in progress. */
-	uint32 BuildProcessID;
-
-	/**
-	* A map of directory paths to shader jobs contained within that directory.
-	* One entry per XGE task.
-	*/
-	class FShaderBatch
-	{
-		TArray<FShaderCommonCompileJobPtr> Jobs;
-		bool bTransferFileWritten;
-
-	public:
-		bool bSuccessfullyCompleted;
-		const FString& DirectoryBase;
-		const FString InputFileName;
-		const FString& SuccessFileName;
-		const FString OutputFileName;
-
-		int32 BatchIndex;
-		int32 DirectoryIndex;
-
-		FString WorkingDirectory;
-		FString OutputFileNameAndPath;
-		FString SuccessFileNameAndPath;
-		FString InputFileNameAndPath;
-
-		FShaderBatch(const FString& InDirectoryBase, const FString& InInputFileName, const FString& InSuccessFileName, const FString& InOutputFileName, int32 InDirectoryIndex, int32 InBatchIndex)
-			: bTransferFileWritten(false)
-			, bSuccessfullyCompleted(false)
-			, DirectoryBase(InDirectoryBase)
-			, InputFileName(InInputFileName)
-			, SuccessFileName(InSuccessFileName)
-			, OutputFileName(InOutputFileName)
-		{
-			SetIndices(InDirectoryIndex, InBatchIndex);
-		}
-
-		void SetIndices(int32 InDirectoryIndex, int32 InBatchIndex);
-
-		void CleanUpFiles(bool keepInputFile);
-
-		inline int32 NumJobs()
-		{
-			return Jobs.Num();
-		}
-		inline const TArray<FShaderCommonCompileJobPtr>& GetJobs() const
-		{
-			return Jobs;
-		}
-
-		void AddJob(FShaderCommonCompileJobPtr Job);
-
-		void WriteTransferFile();
-	};
-	TArray<FShaderBatch*> ShaderBatchesInFlight;
-	int32 ShaderBatchesInFlightCompleted;
-	TArray<FShaderBatch*> ShaderBatchesFull;
-	TSparseArray<FShaderBatch*> ShaderBatchesIncomplete;
-
-	/** The full path to the two working directories for XGE shader builds. */
-	const FString FASTBuildWorkingDirectory;
-	uint32 FASTBuildDirectoryIndex;
-
-	uint64 LastAddTime;
-	uint64 StartTime;
-	int32 BatchIndexToCreate;
-	int32 BatchIndexToFill;
-
-	FDateTime ScriptFileCreationTime;
-
-	void PostCompletedJobsForBatch(FShaderBatch* Batch);
-
-	void GatherResultsFromFASTBuild();
-
-public:
-	/** Initialization constructor. */
-	FShaderCompileFASTBuildThreadRunnable(class FShaderCompilingManager* InManager);
-	virtual ~FShaderCompileFASTBuildThreadRunnable();
-
-	/** Main work loop. */
-	virtual int32 CompilingLoop() override;
-
-	static bool IsSupported();
-};
-
 /** Results for a single compiled and finalized shader map. */
 using FShaderMapFinalizeResults = FShaderMapCompileResults;
 
