Fix non-unity CIS issue in FASTBuild controller.

diff --git a/Engine/Plugins/FastBuildController/Source/Public/FastBuildControllerModule.h b/Engine/Plugins/FastBuildController/Source/Public/FastBuildControllerModule.h
--- a/Engine/Plugins/FastBuildController/Source/Public/FastBuildControllerModule.h
+++ b/Engine/Plugins/FastBuildController/Source/Public/FastBuildControllerModule.h
@@ -4,6 +4,7 @@
 
 #include "DistributedBuildInterface/Public/DistributedBuildControllerInterface.h"
 #include "Containers/Array.h"
+#include "Containers/Queue.h"
 
 class FFastBuildJobProcessor;
 
diff --git a/Engine/Plugins/FastBuildController/Source/Public/FastBuildJobProcessor.h b/Engine/Plugins/FastBuildController/Source/Public/FastBuildJobProcessor.h
--- a/Engine/Plugins/FastBuildController/Source/Public/FastBuildJobProcessor.h
+++ b/Engine/Plugins/FastBuildController/Source/Public/FastBuildJobProcessor.h
@@ -2,7 +2,9 @@
 
 #pragma once
 
+#include "CoreMinimal.h"
 #include "HAL/Runnable.h"
+#include "HAL/RunnableThread.h"
 #include "HAL/ThreadSafeBool.h"
 
 class FFastBuildControllerModule;
diff --git a/Engine/Plugins/FastBuildController/Source/Public/FastBuildUtilities.h b/Engine/Plugins/FastBuildController/Source/Public/FastBuildUtilities.h
--- a/Engine/Plugins/FastBuildController/Source/Public/FastBuildUtilities.h
+++ b/Engine/Plugins/FastBuildController/Source/Public/FastBuildUtilities.h
@@ -2,6 +2,8 @@
 
 #pragma once
 
+#include "HAL/PlatformFile.h"
+
 struct FTask;
 
 class FDependencyEnumerator : public IPlatformFile::FDirectoryVisitor
