More non-unity CIS fix for FastBuild.

diff --git a/Engine/Plugins/FastBuildController/Source/Public/FastBuildUtilities.h b/Engine/Plugins/FastBuildController/Source/Public/FastBuildUtilities.h
--- a/Engine/Plugins/FastBuildController/Source/Public/FastBuildUtilities.h
+++ b/Engine/Plugins/FastBuildController/Source/Public/FastBuildUtilities.h
@@ -2,6 +2,7 @@
 
 #pragma once
 
+#include "CoreMinimal.h"
 #include "HAL/PlatformFile.h"
 
 struct FTask;
