Update macOS metal compiler target version to match MetalShaderFormat.

#rnx

diff --git a/Engine/Plugins/FastBuildController/Source/Private/FastBuildControllerModule.cpp b/Engine/Plugins/FastBuildController/Source/Private/FastBuildControllerModule.cpp
--- a/Engine/Plugins/FastBuildController/Source/Private/FastBuildControllerModule.cpp
+++ b/Engine/Plugins/FastBuildController/Source/Private/FastBuildControllerModule.cpp
@@ -30,7 +30,7 @@ namespace FASTBuildControllerVariables
 static FString GetMetalCompilerFolder()
 {
 	FString Result;
-	if (FPlatformProcess::ExecProcess(TEXT("/usr/bin/xcrun"), TEXT("--sdk macosx metal -v --target=air64-apple-darwin19.6.0"), nullptr, &Result, &Result))
+	if (FPlatformProcess::ExecProcess(TEXT("/usr/bin/xcrun"), TEXT("--sdk macosx metal -v --target=air64-apple-darwin18.7.0"), nullptr, &Result, &Result))
 	{
 		const TCHAR InstalledDirText[] = TEXT("InstalledDir:");
 		int32 InstalledDirOffset = Result.Find(InstalledDirText, ESearchCase::CaseSensitive);
@@ -152,7 +152,6 @@ bool FFastBuildControllerModule::IsSupported()
 			UE_LOG(LogFastBuildController, Warning, TEXT("FASTBuild Shader Compiler is temporarily disabled on Mac until problems with locating Autogen shader folder are sorted out"));
 			FASTBuildControllerVariables::Enabled = 0;
 		}
-
 #endif // PLATFORM_MAC
 	}
 
