PR #9455: Merged FASTBuild PR #8784 with recent UE5 changes (Contributed by [REDACTED_PERSON])

diff --git a/Engine/Plugins/FastBuildController/Source/Private/FastBuildControllerModule.cpp b/Engine/Plugins/FastBuildController/Source/Private/FastBuildControllerModule.cpp
--- a/Engine/Plugins/FastBuildController/Source/Private/FastBuildControllerModule.cpp
+++ b/Engine/Plugins/FastBuildController/Source/Private/FastBuildControllerModule.cpp
@@ -5,6 +5,7 @@
 #include "HAL/FileManager.h"
 #include "HAL/PlatformFileManager.h"
 #include "HAL/PlatformFile.h"
+#include "Misc/App.h"
 #include "Misc/CommandLine.h"
 #include "Misc/ScopeLock.h"
 #include "Misc/Paths.h"
@@ -59,7 +60,10 @@ FFastBuildControllerModule::FFastBuildControllerModule()
 	, RootWorkingDirectory(FString::Printf(TEXT("%sUnrealFastBuildWorkingDir/"), FPlatformProcess::UserTempDir()))
 	, WorkingDirectory(RootWorkingDirectory + FGuid::NewGuid().ToString(EGuidFormats::Digits))
 	, bShutdown(false)
-{}
+{
+	const FGuid Guid = FGuid::NewGuid();
+	IntermediateShadersDirectory = FPaths::EngineIntermediateDir() / TEXT("Shaders") / TEXT("FB") / FApp::GetProjectName();
+}
 
 FFastBuildControllerModule::~FFastBuildControllerModule()
 {
@@ -222,6 +226,8 @@ void FFastBuildControllerModule::InitializeController()
 		IFileManager::Get().MakeDirectory(*RootWorkingDirectory);
 	}
 
+	IFileManager::Get().DeleteDirectory(*IntermediateShadersDirectory, false, true);
+
 	if (IsSupported())
 	{
 		JobDispatcherThread = MakeUnique<FFastBuildJobProcessor>(*this);
@@ -231,6 +237,39 @@ void FFastBuildControllerModule::InitializeController()
 	bControllerInitialized = true;
 }
 
+FString FFastBuildControllerModule::RemapPath(const FString& SourcePath) const
+{
+	if (!FPaths::IsUnderDirectory(SourcePath, FPaths::RootDir()))
+	{
+		FString DestinationPath = SourcePath;
+		FString AbsoluteProjectDir = FPaths::ConvertRelativePathToFull(FPaths::ProjectDir());
+		if (DestinationPath.StartsWith(AbsoluteProjectDir))
+		{
+			DestinationPath = DestinationPath.Mid(AbsoluteProjectDir.Len());
+		}
+		if (DestinationPath.StartsWith(TEXT("\\\\"), ESearchCase::CaseSensitive))
+		{
+			DestinationPath = TEXT("____") + DestinationPath.Mid(2);
+		}
+		DestinationPath.ReplaceInline(TEXT("\\"), TEXT("/"), ESearchCase::CaseSensitive);
+		if (DestinationPath.Len() > 1 && DestinationPath[1] == ':')
+		{
+			DestinationPath = DestinationPath.Left(1) + TEXT("__") / DestinationPath.Mid(2);
+		}
+		if (DestinationPath.Len() > 0 && DestinationPath[0] == '/')
+		{
+			DestinationPath = TEXT("__") + DestinationPath.Mid(1);
+		}
+		DestinationPath = IntermediateShadersDirectory / DestinationPath;
+
+		return DestinationPath;
+	}
+	else
+	{
+		return SourcePath;
+	}
+}
+
 void FFastBuildControllerModule::CleanWorkingDirectory()
 {
 	IFileManager& FileManager = IFileManager::Get();
@@ -240,6 +279,14 @@ void FFastBuildControllerModule::CleanWorkingDirectory()
 		if (!FileManager.DeleteDirectory(*WorkingDirectory, false, true))
 		{
 			UE_LOG(LogFastBuildController, Log, TEXT("%s => Failed to delete current working Directory => %s"), ANSI_TO_TCHAR(__FUNCTION__), *RootWorkingDirectory);
+		}
+	}
+
+	if (!IntermediateShadersDirectory.IsEmpty())
+	{
+		if (!FileManager.DeleteDirectory(*IntermediateShadersDirectory, false, true))
+		{
+			UE_LOG(LogFastBuildController, Log, TEXT("%s => Failed to delete directory => %s"), ANSI_TO_TCHAR(__FUNCTION__), *IntermediateShadersDirectory);
 		}	
 	}
 }
diff --git a/Engine/Plugins/FastBuildController/Source/Private/FastBuildUtilities.cpp b/Engine/Plugins/FastBuildController/Source/Private/FastBuildUtilities.cpp
--- a/Engine/Plugins/FastBuildController/Source/Private/FastBuildUtilities.cpp
+++ b/Engine/Plugins/FastBuildController/Source/Private/FastBuildUtilities.cpp
@@ -1,7 +1,7 @@
 ﻿// Copyright Epic Games, Inc. All Rights Reserved.
 
 #include "FastBuildUtilities.h"
-#include <DistributedBuildInterface/Public/DistributedBuildControllerInterface.h>
+#include "FastBuildControllerModule.h"
 
 #include "ShaderCore.h"
 #include "HAL/PlatformFileManager.h"
@@ -20,7 +20,7 @@ namespace FASTBuildControllerUtilitiesVariables
         TEXT("1: Send along debug information in FASTBuild."),
         ECVF_Default);
 
-	int32 SendAllPossibleShaderDependencies = 0;
+	int32 SendAllPossibleShaderDependencies = 1;
 	FAutoConsoleVariableRef CVarFASTBuildSendAllPossibleShaderDependencies(
         TEXT("r.FASTBuildController.SendAllPossibleShaderDependencies"),
         SendAllPossibleShaderDependencies,
@@ -152,13 +152,13 @@ void FastBuildUtilities::WritePlatformCompilerDependenciesToFile(FArchive& Scrip
 
 void FastBuildUtilities::WriteDependenciesForShaderToScript(const TArray<FTask*>& InCompilationTasks, FArchive& ScriptFile)
 {
+	TArray<FString> FullUniqueDependenciesArray;
 	if (FASTBuildControllerUtilitiesVariables::SendAllPossibleShaderDependencies)
 	{
 		// This is kinda a hack because we are sending all possible dependencies
-		// Warning: this can result in sending files with the same name located in different paths. FASTBuild currently (as of v1.05) does not handle this, and BFF reading will be broken.
-		FDependencyEnumerator ShaderUsfDeps = FDependencyEnumerator(ScriptFile, nullptr, TEXT(".usf"));
-		FDependencyEnumerator ShaderUshDeps = FDependencyEnumerator(ScriptFile, nullptr, TEXT(".ush"));
-		FDependencyEnumerator ShaderHeaderDeps = FDependencyEnumerator(ScriptFile, nullptr, TEXT(".h"));
+		FDependencyUniqueArrayEnumerator ShaderUsfDeps = FDependencyUniqueArrayEnumerator(FullUniqueDependenciesArray, nullptr, TEXT(".usf"));
+		FDependencyUniqueArrayEnumerator ShaderUshDeps = FDependencyUniqueArrayEnumerator(FullUniqueDependenciesArray, nullptr, TEXT(".ush"));
+		FDependencyUniqueArrayEnumerator ShaderHeaderDeps = FDependencyUniqueArrayEnumerator(FullUniqueDependenciesArray, nullptr, TEXT(".h"));
 		const TMap<FString, FString> ShaderSourceDirectoryMappings = AllShaderSourceDirectoryMappings();
 		for (auto& ShaderDirectoryMapping : ShaderSourceDirectoryMappings)
 		{
@@ -169,22 +169,15 @@ void FastBuildUtilities::WriteDependenciesForShaderToScript(const TArray<FTask*>
 	}
 	else
 	{
-		TArray<FString> FullUniqueDependenciesArray;
 		{
 			for (FTask* CompilationTask : InCompilationTasks)
 			{
 				for (const FString& Dependency : CompilationTask->CommandData.Dependencies)
 				{
-					FullUniqueDependenciesArray.AddUnique(Dependency);
+					FullUniqueDependenciesArray.AddUnique(GetShaderFullPathOfShaderDependency(Dependency));
 				}
 			}
 		}
-
-		for (const FString& ExtraDependency : FullUniqueDependenciesArray)
-		{
-			const FString ExtraFile = TEXT("\t\t'") + GetShaderFullPathOfShaderDependency(ExtraDependency) + TEXT("'," LINE_TERMINATOR_ANSI);
-			ScriptFile.Serialize((void*)StringCast<ANSICHAR>(*ExtraFile, ExtraFile.Len()).Get(), sizeof(ANSICHAR) * ExtraFile.Len());
-		}
 	}
 
 #if PLATFORM_MAC
@@ -192,6 +185,40 @@ void FastBuildUtilities::WriteDependenciesForShaderToScript(const TArray<FTask*>
 	FDependencyEnumerator MetalCompilerDeps = FDependencyEnumerator(ScriptFile, nullptr, nullptr);
 	IFileManager::Get().IterateDirectoryRecursively(*MetalIntermediateDir, MetalCompilerDeps);
 #endif
+
+	for (const FString& ExtraDependency : FullUniqueDependenciesArray)
+	{
+		FString SourcePath = ExtraDependency;
+		FString DestinationPath;
+		if (!FPaths::IsUnderDirectory(SourcePath, FPaths::RootDir()))
+		{
+			DestinationPath = FFastBuildControllerModule::Get().RemapPath(SourcePath);
+
+			IPlatformFile& PlatformFile = FPlatformFileManager::Get().GetPlatformFile();
+			FDateTime SourceTimeStamp = PlatformFile.GetTimeStamp(*SourcePath);
+			FDateTime DestinationTimeStamp = PlatformFile.GetTimeStamp(*DestinationPath);
+
+			if (SourceTimeStamp != FDateTime::MinValue() && SourceTimeStamp != DestinationTimeStamp)
+			{
+				FString DestinationDirectory = FPaths::GetPath(DestinationPath);
+				if (!PlatformFile.DirectoryExists(*DestinationDirectory))
+				{
+					PlatformFile.CreateDirectoryTree(*DestinationDirectory);
+				}
+				PlatformFile.CopyFile(*DestinationPath, *SourcePath);
+				PlatformFile.SetTimeStamp(*DestinationPath, SourceTimeStamp);
+			}
+
+			DestinationPath = IFileManager::Get().ConvertToAbsolutePathForExternalAppForWrite(*DestinationPath);
+		}
+		else
+		{
+			DestinationPath = SourcePath;
+		}
+
+		const FString ExtraFile = TEXT("\t\t'") + DestinationPath + TEXT("'," LINE_TERMINATOR_ANSI);
+		ScriptFile.Serialize((void*)StringCast<ANSICHAR>(*ExtraFile, ExtraFile.Len()).Get(), sizeof(ANSICHAR) * ExtraFile.Len());
+	}
 }
 
 void FastBuildUtilities::FASTBuildWriteScriptFileHeader(const TArray<FTask*>& InCompilationTasks, FArchive& ScriptFile, const FString& WorkerName)
@@ -258,14 +285,14 @@ FString FastBuildUtilities::GetFastBuildExecutablePath()
 {
 	FString FASTBuildExecutablePath = TEXT("PLATFORM NOT SUPPORTED");
 #if PLATFORM_WINDOWS
-	FASTBuildExecutablePath = TEXT("Extras\\ThirdPartyNotUE\\FASTBuild\\Win64\\FBuild.exe");
+	FASTBuildExecutablePath = FPaths::EngineDir() / TEXT("Extras\\ThirdPartyNotUE\\FASTBuild\\Win64\\FBuild.exe");
 #elif PLATFORM_MAC
-	FASTBuildExecutablePath = TEXT("Extras/ThirdPartyNotUE/FASTBuild/Mac/FBuild");
+	FASTBuildExecutablePath = FPaths::EngineDir() / TEXT("Extras/ThirdPartyNotUE/FASTBuild/Mac/FBuild");
 #elif PLATFORM_LINUX
-	FASTBuildExecutablePath = TEXT("Extras/ThirdPartyNotUE/FASTBuild/Linux/fbuild");
+	FASTBuildExecutablePath = FPaths::EngineDir() / TEXT("Extras/ThirdPartyNotUE/FASTBuild/Linux/fbuild");
 #endif
 
-	return FPaths::EngineDir() / FASTBuildExecutablePath;
+	return FASTBuildExecutablePath;
 }
 
 FString FastBuildUtilities::GetFastBuildCache()
diff --git a/Engine/Plugins/FastBuildController/Source/Public/FastBuildControllerModule.h b/Engine/Plugins/FastBuildController/Source/Public/FastBuildControllerModule.h
--- a/Engine/Plugins/FastBuildController/Source/Public/FastBuildControllerModule.h
+++ b/Engine/Plugins/FastBuildController/Source/Public/FastBuildControllerModule.h
@@ -35,6 +35,8 @@ class FFastBuildControllerModule : public IDistributedBuildController
 
 	const FString WorkingDirectory;
 
+	FString IntermediateShadersDirectory;
+
 	bool bShutdown;
 
 	void CleanWorkingDirectory();
@@ -52,10 +54,13 @@ public:
 
 	virtual const FString GetName() override final { return FString("FastBuild Controller"); };
 
+	virtual FString RemapPath(const FString& SourcePath) const override final;
+
 	virtual bool IsSupported() override final;
 
 	virtual FString CreateUniqueFilePath() override final;
 	virtual TFuture<FDistributedBuildTaskResult> EnqueueTask(const FTaskCommandData& CommandData) override final;
+
 	void EnqueueTask(FTask* Task);
 	FTask* DequeueTask();
 
@@ -90,6 +95,8 @@ public:
 
 	FString GetWorkingDirectory() const { return WorkingDirectory;}
 
+	FString GetIntermediateShadersDirectory() const { return IntermediateShadersDirectory; }
+
 	const TMap<uint32, FTask*>& GetDispatchedTasks() const
 	{
 		return DispatchedTasks;
diff --git a/Engine/Plugins/FastBuildController/Source/Public/FastBuildUtilities.h b/Engine/Plugins/FastBuildController/Source/Public/FastBuildUtilities.h
--- a/Engine/Plugins/FastBuildController/Source/Public/FastBuildUtilities.h
+++ b/Engine/Plugins/FastBuildController/Source/Public/FastBuildUtilities.h
@@ -40,6 +40,36 @@ public:
 	const TCHAR* Extension;
 };
 
+class FDependencyUniqueArrayEnumerator : public IPlatformFile::FDirectoryVisitor
+{
+public:
+	FDependencyUniqueArrayEnumerator(TArray<FString>& InFullUniqueDependenciesArray, const TCHAR* InPrefix, const TCHAR* InExtension)
+        : FullUniqueDependenciesArray(InFullUniqueDependenciesArray)
+        , Prefix(InPrefix)
+        , Extension(InExtension)
+	{
+	}
+
+	virtual bool Visit(const TCHAR* FilenameChar, bool bIsDirectory) override
+	{
+		if (!bIsDirectory)
+		{
+			const FString Filename = FString(FilenameChar);
+
+			if ((!Prefix || Filename.Contains(Prefix)) && (!Extension || Filename.EndsWith(Extension)))
+			{
+				FullUniqueDependenciesArray.AddUnique(*IFileManager::Get().ConvertToAbsolutePathForExternalAppForWrite(*Filename));
+			}
+		}
+
+		return true;
+	}
+
+	TArray<FString>& FullUniqueDependenciesArray;
+	const TCHAR* Prefix;
+	const TCHAR* Extension;
+};
+
 class FastBuildUtilities
 {
 public:
diff --git a/Engine/Source/Developer/DistributedBuildInterface/Public/DistributedBuildControllerInterface.h b/Engine/Source/Developer/DistributedBuildInterface/Public/DistributedBuildControllerInterface.h
--- a/Engine/Source/Developer/DistributedBuildInterface/Public/DistributedBuildControllerInterface.h
+++ b/Engine/Source/Developer/DistributedBuildInterface/Public/DistributedBuildControllerInterface.h
@@ -57,6 +57,8 @@ public:
 	// Returns the name of the controller. Used for logging purposes.
 	virtual const FString GetName() = 0;
 
+	virtual FString RemapPath(const FString& SourcePath) const { return SourcePath; }
+
 	virtual void Tick(float DeltaSeconds){}
 
 	// Returns a new file path to be used for writing input data to.
diff --git a/Engine/Source/Runtime/Engine/Private/ShaderCompiler/ShaderCompiler.cpp b/Engine/Source/Runtime/Engine/Private/ShaderCompiler/ShaderCompiler.cpp
--- a/Engine/Source/Runtime/Engine/Private/ShaderCompiler/ShaderCompiler.cpp
+++ b/Engine/Source/Runtime/Engine/Private/ShaderCompiler/ShaderCompiler.cpp
@@ -1290,7 +1290,7 @@ static void SplitJobsByType(const TArray<FShaderCommonCompileJobPtr>& QueuedJobs
 }
 
 // Serialize Queued Job information
-bool FShaderCompileUtilities::DoWriteTasks(const TArray<FShaderCommonCompileJobPtr>& QueuedJobs, FArchive& InTransferFile, bool bUseRelativePaths, bool bCompressTaskFile)
+bool FShaderCompileUtilities::DoWriteTasks(const TArray<FShaderCommonCompileJobPtr>& QueuedJobs, FArchive& InTransferFile, IDistributedBuildController* BuildDistributionController, bool bUseRelativePaths, bool bCompressTaskFile)
 {
 	int32 InputVersion = ShaderCompileWorkerInputVersion;
 	InTransferFile << InputVersion;
@@ -1313,8 +1313,25 @@ bool FShaderCompileUtilities::DoWriteTasks(const TArray<FShaderCommonCompileJobP
 	TMap<FString, FString> ShaderSourceDirectoryMappings = AllShaderSourceDirectoryMappings();
 	for(TPair<FString, FString>& Pair : ShaderSourceDirectoryMappings)
 	{
-		// enforce relative paths when bUseRelativePaths=true
-		Pair.Value = bUseRelativePaths ? FPaths::CreateStandardFilename(Pair.Value) : FPaths::ConvertRelativePathToFull(Pair.Value);
+		// Remap/enforce relative paths when bUseRelativePaths=true
+		if (bUseRelativePaths && BuildDistributionController != nullptr)
+		{
+			FString SourcePath = FPaths::ConvertRelativePathToFull(Pair.Value);
+			if (!FPaths::IsUnderDirectory(SourcePath, FPaths::RootDir()))
+			{
+				FString DestinationPath = BuildDistributionController->RemapPath(SourcePath);
+				DestinationPath = FPaths::CreateStandardFilename(DestinationPath);
+				Pair.Value = DestinationPath;
+			}
+			else
+			{
+				Pair.Value = FPaths::CreateStandardFilename(Pair.Value);
+			}
+		}
+		else
+		{
+			Pair.Value = FPaths::ConvertRelativePathToFull(Pair.Value);
+		}
 	}
 	TransferFile << ShaderSourceDirectoryMappings;
 
diff --git a/Engine/Source/Runtime/Engine/Private/ShaderCompiler/ShaderCompilerDistributed.cpp b/Engine/Source/Runtime/Engine/Private/ShaderCompiler/ShaderCompilerDistributed.cpp
--- a/Engine/Source/Runtime/Engine/Private/ShaderCompiler/ShaderCompilerDistributed.cpp
+++ b/Engine/Source/Runtime/Engine/Private/ShaderCompiler/ShaderCompilerDistributed.cpp
@@ -77,7 +77,7 @@ void FShaderCompileDistributedThreadRunnable_Interface::DispatchShaderCompileJob
 	// Serialize the jobs to the input file
 	GShaderCompilerStats->RegisterJobBatch(JobsToSerialize.Num(), FShaderCompilerStats::EExecutionType::Distributed);
 	FArchive* InputFileAr = IFileManager::Get().CreateFileWriter(*InputFilePath, FILEWRITE_EvenIfReadOnly | FILEWRITE_NoFail);
-	FShaderCompileUtilities::DoWriteTasks(JobsToSerialize, *InputFileAr, CachedController.RequiresRelativePaths());
+	FShaderCompileUtilities::DoWriteTasks(JobsToSerialize, *InputFileAr, &CachedController, CachedController.RequiresRelativePaths());
 	delete InputFileAr;
 
 	// Kick off the job
diff --git a/Engine/Source/Runtime/Engine/Public/ShaderCompiler.h b/Engine/Source/Runtime/Engine/Public/ShaderCompiler.h
--- a/Engine/Source/Runtime/Engine/Public/ShaderCompiler.h
+++ b/Engine/Source/Runtime/Engine/Public/ShaderCompiler.h
@@ -609,7 +609,7 @@ private:
 
 namespace FShaderCompileUtilities
 {
-	bool DoWriteTasks(const TArray<FShaderCommonCompileJobPtr>& QueuedJobs, FArchive& TransferFile, bool bUseRelativePaths = false, bool bCompressTaskFile = false);
+	bool DoWriteTasks(const TArray<FShaderCommonCompileJobPtr>& QueuedJobs, FArchive& TransferFile, IDistributedBuildController* BuildDistributionController = nullptr, bool bUseRelativePaths = false, bool bCompressTaskFile = false);
 	void DoReadTaskResults(const TArray<FShaderCommonCompileJobPtr>& QueuedJobs, FArchive& OutputFile);
 
 	/** Execute the specified (single or pipeline) shader compile job. */
diff --git a/Engine/Source/Runtime/RenderCore/Private/ShaderCore.cpp b/Engine/Source/Runtime/RenderCore/Private/ShaderCore.cpp
--- a/Engine/Source/Runtime/RenderCore/Private/ShaderCore.cpp
+++ b/Engine/Source/Runtime/RenderCore/Private/ShaderCore.cpp
@@ -1738,7 +1738,12 @@ void AddShaderSourceDirectoryMapping(const FString& VirtualShaderDirectory, cons
 	check(!GShaderSourceDirectoryMappings.Contains(VirtualShaderDirectory));
 
 	// Make sure the real directory to map exists.
-	checkf(FPaths::DirectoryExists(RealShaderDirectory), TEXT("FPaths::DirectoryExists(%s)"), *RealShaderDirectory);
+	bool bDirectoryExists = FPaths::DirectoryExists(RealShaderDirectory);
+	if (!bDirectoryExists)
+	{
+		UE_LOG(LogShaders, Log, TEXT("Directory %s"), *RealShaderDirectory);
+	}
+	checkf(bDirectoryExists, TEXT("FPaths::DirectoryExists(%s %s) %s"), *RealShaderDirectory, *FPaths::ConvertRelativePathToFull(RealShaderDirectory), FPlatformProcess::ComputerName());
 
 	// Make sure the Generated directory does not exist, because is reserved for C++ generated shader source
 	// by the FShaderCompilerEnvironment::IncludeVirtualPathToContentsMap member.
