Fixed FastBuild failing to distribute .hlsli files (included in TraceRayInline.ush).

diff --git a/Engine/Plugins/FastBuildController/Source/Private/FastBuildUtilities.cpp b/Engine/Plugins/FastBuildController/Source/Private/FastBuildUtilities.cpp
--- a/Engine/Plugins/FastBuildController/Source/Private/FastBuildUtilities.cpp
+++ b/Engine/Plugins/FastBuildController/Source/Private/FastBuildUtilities.cpp
@@ -160,12 +160,19 @@ void FastBuildUtilities::WriteDependenciesForShaderToScript(const TArray<FTask*>
 		FDependencyUniqueArrayEnumerator ShaderUsfDeps = FDependencyUniqueArrayEnumerator(FullUniqueDependenciesArray, nullptr, TEXT(".usf"));
 		FDependencyUniqueArrayEnumerator ShaderUshDeps = FDependencyUniqueArrayEnumerator(FullUniqueDependenciesArray, nullptr, TEXT(".ush"));
 		FDependencyUniqueArrayEnumerator ShaderHeaderDeps = FDependencyUniqueArrayEnumerator(FullUniqueDependenciesArray, nullptr, TEXT(".h"));
+#if PLATFORM_WINDOWS
+		FDependencyUniqueArrayEnumerator ShaderHlsliDeps = FDependencyUniqueArrayEnumerator(FullUniqueDependenciesArray, nullptr, TEXT(".hlsli"));
+#endif
+
 		const TMap<FString, FString> ShaderSourceDirectoryMappings = AllShaderSourceDirectoryMappings();
 		for (auto& ShaderDirectoryMapping : ShaderSourceDirectoryMappings)
 		{
 			IFileManager::Get().IterateDirectoryRecursively(*ShaderDirectoryMapping.Value, ShaderUsfDeps);
 			IFileManager::Get().IterateDirectoryRecursively(*ShaderDirectoryMapping.Value, ShaderUshDeps);
 			IFileManager::Get().IterateDirectoryRecursively(*ShaderDirectoryMapping.Value, ShaderHeaderDeps);
+#if PLATFORM_WINDOWS
+			IFileManager::Get().IterateDirectoryRecursively(*ShaderDirectoryMapping.Value, ShaderHlsliDeps);
+#endif
 		}
 	}
 	else
