[Shaders] Minor code cleanup of canceling distributed shader compile tasks.

#rnx

diff --git a/Engine/Plugins/FastBuildController/Source/Private/FastBuildControllerModule.cpp b/Engine/Plugins/FastBuildController/Source/Private/FastBuildControllerModule.cpp
--- a/Engine/Plugins/FastBuildController/Source/Private/FastBuildControllerModule.cpp
+++ b/Engine/Plugins/FastBuildController/Source/Private/FastBuildControllerModule.cpp
@@ -191,22 +191,15 @@ void FFastBuildControllerModule::ShutdownModule()
 	// Cancel any remaining tasks
 	for (auto Iter = DispatchedTasks.CreateIterator(); Iter; ++Iter)
 	{
-		FDistributedBuildTaskResult Result;
-		Result.ReturnCode = 0;
-		Result.bCompleted = false;
-
 		FTask* Task = Iter.Value();
-		Task->Promise.SetValue(Result);
+		Task->Cancel();
 		delete Task;
 	}
 
 	FTask* Task;
 	while (PendingTasks.Dequeue(Task))
 	{
-		FDistributedBuildTaskResult Result;
-		Result.ReturnCode = 0;
-		Result.bCompleted = false;
-		Task->Promise.SetValue(Result);
+		Task->Cancel();
 		delete Task;
 	}
 
@@ -359,10 +352,7 @@ void FFastBuildControllerModule::DeRegisterDispatchedTasks(const TArray<uint32>&
 
 void FFastBuildControllerModule::ReportJobProcessed(FTask* CompletedTask, const FTaskResponse& InTaskResponse)
 {
-	FDistributedBuildTaskResult Result;
-	Result.ReturnCode = InTaskResponse.ReturnCode;
-	Result.bCompleted = true;
-	CompletedTask->Promise.SetValue(Result);
+	CompletedTask->Finalize(InTaskResponse.ReturnCode);
 	delete CompletedTask;
 }
 
diff --git a/Engine/Plugins/UbaController/Source/Private/UbaControllerModule.cpp b/Engine/Plugins/UbaController/Source/Private/UbaControllerModule.cpp
--- a/Engine/Plugins/UbaController/Source/Private/UbaControllerModule.cpp
+++ b/Engine/Plugins/UbaController/Source/Private/UbaControllerModule.cpp
@@ -219,10 +219,7 @@ void FUbaControllerModule::ShutdownModule()
 		FTask* Task;
 		while (PendingRequestedCompilationTasks.Dequeue(Task))
 		{
-			FDistributedBuildTaskResult Result;
-			Result.ReturnCode = 0;
-			Result.bCompleted = false;
-			Task->Promise.SetValue(Result);
+			Task->Cancel();
 			delete Task;
 		}
 
@@ -300,10 +297,7 @@ void FUbaControllerModule::ReportJobProcessed(const FTaskResponse& InTaskRespons
 {
 	if (CompileTask)
 	{
-		FDistributedBuildTaskResult Result;
-		Result.ReturnCode = InTaskResponse.ReturnCode;
-		Result.bCompleted = true;
-		CompileTask->Promise.SetValue(Result);
+		CompileTask->Finalize(InTaskResponse.ReturnCode);
 		delete CompileTask;
 	}
 }
diff --git a/Engine/Plugins/XGEController/Source/Private/XGEControllerModule.cpp b/Engine/Plugins/XGEController/Source/Private/XGEControllerModule.cpp
--- a/Engine/Plugins/XGEController/Source/Private/XGEControllerModule.cpp
+++ b/Engine/Plugins/XGEController/Source/Private/XGEControllerModule.cpp
@@ -313,22 +313,15 @@ void FXGEControllerModule::ShutdownModule()
 		// Cancel any remaining tasks
 		for (auto Iter = DispatchedTasks.CreateIterator(); Iter; ++Iter)
 		{
-			FDistributedBuildTaskResult Result;
-			Result.ReturnCode = 0;
-			Result.bCompleted = false;
-
 			FTask* Task = Iter.Value();
-			Task->Promise.SetValue(Result);
+			Task->Cancel();
 			delete Task;
 		}
 
 		FTask* Task;
 		while (PendingTasks.Dequeue(Task))
 		{
-			FDistributedBuildTaskResult Result;
-			Result.ReturnCode = 0;
-			Result.bCompleted = false;
-			Task->Promise.SetValue(Result);
+			Task->Cancel();
 			delete Task;
 		}
 
@@ -559,11 +552,7 @@ void FXGEControllerModule::ReadBackThreadProc()
 				Task = DispatchedTasks.FindAndRemoveChecked(CompletedTaskResponse.ID);
 			}
 
-			FDistributedBuildTaskResult Result;
-			Result.ReturnCode = CompletedTaskResponse.ReturnCode;
-			Result.bCompleted = true;
-
-			Task->Promise.SetValue(Result);
+			Task->Finalize(CompletedTaskResponse.ReturnCode);
 			delete Task;
 		}
 
diff --git a/Engine/Source/Developer/DistributedBuildInterface/Public/DistributedBuildControllerInterface.h b/Engine/Source/Developer/DistributedBuildInterface/Public/DistributedBuildControllerInterface.h
--- a/Engine/Source/Developer/DistributedBuildInterface/Public/DistributedBuildControllerInterface.h
+++ b/Engine/Source/Developer/DistributedBuildInterface/Public/DistributedBuildControllerInterface.h
@@ -38,10 +38,29 @@ struct FTask
 	TPromise<FDistributedBuildTaskResult> Promise;
 
 	FTask(uint32 ID, const FTaskCommandData& CommandData, TPromise<FDistributedBuildTaskResult>&& Promise)
-        : ID(ID)
-        , CommandData(CommandData)
-        , Promise(MoveTemp(Promise))
-	{}
+		: ID(ID)
+		, CommandData(CommandData)
+		, Promise(MoveTemp(Promise))
+	{
+	}
+
+	/** Sets the promised task result to being incomplete, i.e. FDistributedBuildTaskResult::bCompleted=false. */
+	void Cancel()
+	{
+		FDistributedBuildTaskResult Result;
+		Result.ReturnCode = 0;
+		Result.bCompleted = false;
+		Promise.SetValue(Result);
+	}
+
+	/** Sets the promised task result to being completed with the specified return code. */
+	void Finalize(int32 InReturnCode)
+	{
+		FDistributedBuildTaskResult Result;
+		Result.ReturnCode = InReturnCode;
+		Result.bCompleted = true;
+		Promise.SetValue(Result);
+	}
 };
 
 struct FTaskResponse
