[Shaders] Renamed FTask in DistributedBuildControllerInterface.h to FDistributedBuildTask to avoid confusion and compiler ambiguity with UE::Tasks::FTask.

#rnx

diff --git a/Engine/Plugins/FastBuildController/Source/Private/FastBuildControllerModule.cpp b/Engine/Plugins/FastBuildController/Source/Private/FastBuildControllerModule.cpp
--- a/Engine/Plugins/FastBuildController/Source/Private/FastBuildControllerModule.cpp
+++ b/Engine/Plugins/FastBuildController/Source/Private/FastBuildControllerModule.cpp
@@ -191,12 +191,12 @@ void FFastBuildControllerModule::ShutdownModule()
 	// Cancel any remaining tasks
 	for (auto Iter = DispatchedTasks.CreateIterator(); Iter; ++Iter)
 	{
-		FTask* Task = Iter.Value();
+		FDistributedBuildTask* Task = Iter.Value();
 		Task->Cancel();
 		delete Task;
 	}
 
-	FTask* Task;
+	FDistributedBuildTask* Task;
 	while (PendingTasks.Dequeue(Task))
 	{
 		Task->Cancel();
@@ -288,7 +288,7 @@ TFuture<FDistributedBuildTaskResult> FFastBuildControllerModule::EnqueueTask(con
 	TFuture<FDistributedBuildTaskResult> Future = Promise.GetFuture();
 
 	// Enqueue the new task
-	FTask* Task = new FTask(NextTaskID.Increment(), CommandData, MoveTemp(Promise));
+	FDistributedBuildTask* Task = new FDistributedBuildTask(NextTaskID.Increment(), CommandData, MoveTemp(Promise));
 	{
 		FScopeLock Lock(TasksCS.Get());
 		PendingTasks.Enqueue(Task);
@@ -298,7 +298,7 @@ TFuture<FDistributedBuildTaskResult> FFastBuildControllerModule::EnqueueTask(con
 	return MoveTemp(Future);
 }
 
-void FFastBuildControllerModule::EnqueueTask(FTask* Task)
+void FFastBuildControllerModule::EnqueueTask(FDistributedBuildTask* Task)
 {
 	FScopeLock Lock(TasksCS.Get());
 	{
@@ -307,9 +307,9 @@ void FFastBuildControllerModule::EnqueueTask(FTask* Task)
 	PendingTasksCounter.Increment();
 }
 
-FTask* FFastBuildControllerModule::DequeueTask()
+FDistributedBuildTask* FFastBuildControllerModule::DequeueTask()
 {
-	FTask* TaskPtr;
+	FDistributedBuildTask* TaskPtr;
 	
 	FScopeLock Lock(TasksCS.Get());
 	{
@@ -320,7 +320,7 @@ FTask* FFastBuildControllerModule::DequeueTask()
 	return TaskPtr;
 }
 
-void FFastBuildControllerModule::RegisterDispatchedTask(FTask* DispatchedTask)
+void FFastBuildControllerModule::RegisterDispatchedTask(FDistributedBuildTask* DispatchedTask)
 {
 	check(DispatchedTask)
 	{
@@ -333,7 +333,7 @@ void FFastBuildControllerModule::ReEnqueueDispatchedTasks()
 {
 	FScopeLock Lock(TasksCS.Get());
 	// Reclaim dispatched (incomplete) tasks 
-	for (const TPair<uint32, FTask*>& DispatchedTaskEntry : DispatchedTasks)
+	for (const TPair<uint32, FDistributedBuildTask*>& DispatchedTaskEntry : DispatchedTasks)
 	{
 		EnqueueTask(DispatchedTaskEntry.Value);
 	}
@@ -350,7 +350,7 @@ void FFastBuildControllerModule::DeRegisterDispatchedTasks(const TArray<uint32>&
 	}
 }
 
-void FFastBuildControllerModule::ReportJobProcessed(FTask* CompletedTask, const FTaskResponse& InTaskResponse)
+void FFastBuildControllerModule::ReportJobProcessed(FDistributedBuildTask* CompletedTask, const FTaskResponse& InTaskResponse)
 {
 	CompletedTask->Finalize(InTaskResponse.ReturnCode);
 	delete CompletedTask;
diff --git a/Engine/Plugins/FastBuildController/Source/Private/FastBuildJobProcessor.cpp b/Engine/Plugins/FastBuildController/Source/Private/FastBuildJobProcessor.cpp
--- a/Engine/Plugins/FastBuildController/Source/Private/FastBuildJobProcessor.cpp
+++ b/Engine/Plugins/FastBuildController/Source/Private/FastBuildJobProcessor.cpp
@@ -180,7 +180,7 @@ void FFastBuildJobProcessor::MonitorFastBuildProcess()
 	}
 }
 
-void FFastBuildJobProcessor::WriteScriptFileToDisk(const TArray<FTask*>& PendingTasks, const FString& ScriptFilename, const FString& WorkerName) const
+void FFastBuildJobProcessor::WriteScriptFileToDisk(const TArray<FDistributedBuildTask*>& PendingTasks, const FString& ScriptFilename, const FString& WorkerName) const
 {
 	// Create the FASTBuild script file.
 	TUniquePtr<FArchive> ScriptFile(FastBuildUtilities::CreateFileHelper(ScriptFilename));
@@ -189,7 +189,7 @@ void FFastBuildJobProcessor::WriteScriptFileToDisk(const TArray<FTask*>& Pending
 		FastBuildUtilities::FASTBuildWriteScriptFileHeader(PendingTasks, *ScriptFile, WorkerName);
 		
 		// Write the task line for each shader batch
-		for (const FTask* CompilationTask : PendingTasks)
+		for (const FDistributedBuildTask* CompilationTask : PendingTasks)
 		{
 			FString WorkerAbsoluteDirectory = IFileManager::Get().ConvertToAbsolutePathForExternalAppForWrite(*ControllerModule.GetWorkingDirectory());
 			FPaths::NormalizeDirectoryName(WorkerAbsoluteDirectory);
@@ -219,7 +219,7 @@ void FFastBuildJobProcessor::WriteScriptFileToDisk(const TArray<FTask*>& Pending
         );
 		ScriptFile->Serialize((void*)StringCast<ANSICHAR>(*AliasBuildTargetOpen, AliasBuildTargetOpen.Len()).Get(), sizeof(ANSICHAR) * AliasBuildTargetOpen.Len());
 
-		for (FTask* CompilationTask : PendingTasks)
+		for (FDistributedBuildTask* CompilationTask : PendingTasks)
 		{
 			const FString TargetExport = FString::Printf(TEXT("'ShaderBatch-%d', "), CompilationTask->ID);
 			ScriptFile->Serialize((void*)StringCast<ANSICHAR>(*TargetExport, TargetExport.Len()).Get(), sizeof(ANSICHAR) * TargetExport.Len());
@@ -245,8 +245,8 @@ void FFastBuildJobProcessor::SubmitPendingJobs()
 		return;
 	}
 
-	TArray<FTask*> TasksToSubmit;
-	FTask* PendingTask = nullptr;
+	TArray<FDistributedBuildTask*> TasksToSubmit;
+	FDistributedBuildTask* PendingTask = nullptr;
 	{
 		FScopeLock Lock(ControllerModule.GetTasksCS());
 		while (ControllerModule.AreTasksPending())
@@ -289,9 +289,9 @@ void FFastBuildJobProcessor::GatherBuildResults() const
 	IFileManager& FileManager = IFileManager::Get();
 	
 	TArray<uint32> CompletedTasks;
-	for (const TPair<uint32, FTask*>& DispatchedTaskEntry : ControllerModule.GetDispatchedTasks())
+	for (const TPair<uint32, FDistributedBuildTask*>& DispatchedTaskEntry : ControllerModule.GetDispatchedTasks())
 	{
-		FTask* CompileTask = DispatchedTaskEntry.Value;
+		FDistributedBuildTask* CompileTask = DispatchedTaskEntry.Value;
 
 		constexpr uint64 VersionAndFileSizeSize = sizeof(uint32) + sizeof(uint64);	
 		if (PlatformFile.FileExists(*CompileTask->CommandData.OutputFileName) &&
diff --git a/Engine/Plugins/FastBuildController/Source/Private/FastBuildUtilities.cpp b/Engine/Plugins/FastBuildController/Source/Private/FastBuildUtilities.cpp
--- a/Engine/Plugins/FastBuildController/Source/Private/FastBuildUtilities.cpp
+++ b/Engine/Plugins/FastBuildController/Source/Private/FastBuildUtilities.cpp
@@ -151,7 +151,7 @@ void FastBuildUtilities::WritePlatformCompilerDependenciesToFile(FArchive& Scrip
 	}
 }
 
-void FastBuildUtilities::WriteDependenciesForShaderToScript(const TArray<FTask*>& InCompilationTasks, FArchive& ScriptFile)
+void FastBuildUtilities::WriteDependenciesForShaderToScript(const TArray<FDistributedBuildTask*>& InCompilationTasks, FArchive& ScriptFile)
 {
 	TArray<FString> FullUniqueDependenciesArray;
 	if (FASTBuildControllerUtilitiesVariables::SendAllPossibleShaderDependencies)
@@ -178,7 +178,7 @@ void FastBuildUtilities::WriteDependenciesForShaderToScript(const TArray<FTask*>
 	else
 	{
 		{
-			for (FTask* CompilationTask : InCompilationTasks)
+			for (FDistributedBuildTask* CompilationTask : InCompilationTasks)
 			{
 				for (const FString& Dependency : CompilationTask->CommandData.Dependencies)
 				{
@@ -229,7 +229,7 @@ void FastBuildUtilities::WriteDependenciesForShaderToScript(const TArray<FTask*>
 	}
 }
 
-void FastBuildUtilities::FASTBuildWriteScriptFileHeader(const TArray<FTask*>& InCompilationTasks, FArchive& ScriptFile, const FString& WorkerName)
+void FastBuildUtilities::FASTBuildWriteScriptFileHeader(const TArray<FDistributedBuildTask*>& InCompilationTasks, FArchive& ScriptFile, const FString& WorkerName)
 {	
 	static constexpr TCHAR HeaderTemplate[] =
 		TEXT("Settings" LINE_TERMINATOR_ANSI)
diff --git a/Engine/Plugins/FastBuildController/Source/Public/FastBuildControllerModule.h b/Engine/Plugins/FastBuildController/Source/Public/FastBuildControllerModule.h
--- a/Engine/Plugins/FastBuildController/Source/Public/FastBuildControllerModule.h
+++ b/Engine/Plugins/FastBuildController/Source/Public/FastBuildControllerModule.h
@@ -26,10 +26,10 @@ class FFastBuildControllerModule : public IDistributedBuildController
 	TSharedPtr<FCriticalSection> TasksCS;
 
 	// Queue of tasks submitted by the engine, but not yet dispatched to the controller.
-	TQueue<FTask*> PendingTasks;
+	TQueue<FDistributedBuildTask*> PendingTasks;
 
 	// Map of tasks dispatched to the controller and running within FastBuild, that have not yet finished.
-	TMap<uint32,FTask*> DispatchedTasks;
+	TMap<uint32, FDistributedBuildTask*> DispatchedTasks;
 
 	const FString RootWorkingDirectory;
 
@@ -61,13 +61,13 @@ public:
 	virtual FString CreateUniqueFilePath() override final;
 	virtual TFuture<FDistributedBuildTaskResult> EnqueueTask(const FTaskCommandData& CommandData) override final;
 
-	void EnqueueTask(FTask* Task);
-	FTask* DequeueTask();
+	void EnqueueTask(FDistributedBuildTask* Task);
+	FDistributedBuildTask* DequeueTask();
 
-	void RegisterDispatchedTask(FTask* DispatchedTask);
+	void RegisterDispatchedTask(FDistributedBuildTask* DispatchedTask);
 	void ReEnqueueDispatchedTasks();
 	void DeRegisterDispatchedTasks(const TArray<uint32>& InTasksID);
-	void ReportJobProcessed(FTask* CompletedTask, const FTaskResponse& InTaskResponse);
+	void ReportJobProcessed(FDistributedBuildTask* CompletedTask, const FTaskResponse& InTaskResponse);
 	
 	static FFastBuildControllerModule& Get();
 
@@ -97,7 +97,7 @@ public:
 
 	FString GetIntermediateShadersDirectory() const { return IntermediateShadersDirectory; }
 
-	const TMap<uint32, FTask*>& GetDispatchedTasks() const
+	const TMap<uint32, FDistributedBuildTask*>& GetDispatchedTasks() const
 	{
 		return DispatchedTasks;
 	}
diff --git a/Engine/Plugins/FastBuildController/Source/Public/FastBuildJobProcessor.h b/Engine/Plugins/FastBuildController/Source/Public/FastBuildJobProcessor.h
--- a/Engine/Plugins/FastBuildController/Source/Public/FastBuildJobProcessor.h
+++ b/Engine/Plugins/FastBuildController/Source/Public/FastBuildJobProcessor.h
@@ -10,7 +10,7 @@ class FRunnableThread;
 
 class FFastBuildControllerModule;
 class FShaderBatch;
-struct FTask;
+struct FDistributedBuildTask;
 
 enum EFASTBuild_ReturnCodes
 {
@@ -62,7 +62,7 @@ protected:
 	/** Set to true when the main loop finishes*/
 	FThreadSafeBool bIsWorkDone;
 	
-	void WriteScriptFileToDisk(const TArray<FTask*>& PendingTasks, const FString& ScriptFilename, const FString& WorkerName) const;
+	void WriteScriptFileToDisk(const TArray<FDistributedBuildTask*>& PendingTasks, const FString& ScriptFilename, const FString& WorkerName) const;
 
 	/** Takes all tasks from the queue and creates a FastBuild Script to the kick off the pending jobs */
 	void SubmitPendingJobs();
diff --git a/Engine/Plugins/FastBuildController/Source/Public/FastBuildUtilities.h b/Engine/Plugins/FastBuildController/Source/Public/FastBuildUtilities.h
--- a/Engine/Plugins/FastBuildController/Source/Public/FastBuildUtilities.h
+++ b/Engine/Plugins/FastBuildController/Source/Public/FastBuildUtilities.h
@@ -6,7 +6,7 @@
 #include "HAL/FileManager.h"
 #include "HAL/PlatformFile.h"
 
-struct FTask;
+struct FDistributedBuildTask;
 
 class FDependencyEnumerator : public IPlatformFile::FDirectoryVisitor
 {
@@ -74,7 +74,7 @@ class FastBuildUtilities
 {
 public:
 
-	static void FASTBuildWriteScriptFileHeader(const TArray<FTask*>& InCompilationTasks, FArchive& ScriptFile, const FString& WorkerName);
+	static void FASTBuildWriteScriptFileHeader(const TArray<FDistributedBuildTask*>& InCompilationTasks, FArchive& ScriptFile, const FString& WorkerName);
 	static FArchive* CreateFileHelper(const FString& InFileName);
 
 	static FString GetFastBuildExecutablePath();
@@ -86,5 +86,5 @@ protected:
 	static FString ReplaceEnvironmentVariablesInPath(const FString& ExtraFilePartialPath);
 	static void WriteDependenciesToFileUsingWildcardPath(FArchive& ScriptFile, FString ParsedPath);
 	static void WritePlatformCompilerDependenciesToFile(FArchive& ScriptFile);
-	static void WriteDependenciesForShaderToScript(const TArray<FTask*>& InCompilationTasks, FArchive& ScriptFile);
+	static void WriteDependenciesForShaderToScript(const TArray<FDistributedBuildTask*>& InCompilationTasks, FArchive& ScriptFile);
 };
diff --git a/Engine/Plugins/UbaController/Source/Private/UbaControllerModule.cpp b/Engine/Plugins/UbaController/Source/Private/UbaControllerModule.cpp
--- a/Engine/Plugins/UbaController/Source/Private/UbaControllerModule.cpp
+++ b/Engine/Plugins/UbaController/Source/Private/UbaControllerModule.cpp
@@ -216,7 +216,7 @@ void FUbaControllerModule::ShutdownModule()
 			FPlatformProcess::ConditionalSleep([&](){ return JobDispatcherThread && JobDispatcherThread->IsWorkDone(); },0.1f);
 		}
 
-		FTask* Task;
+		FDistributedBuildTask* Task;
 		while (PendingRequestedCompilationTasks.Dequeue(Task))
 		{
 			Task->Cancel();
@@ -278,7 +278,7 @@ TFuture<FDistributedBuildTaskResult> FUbaControllerModule::EnqueueTask(const FTa
 	TFuture<FDistributedBuildTaskResult> Future = Promise.GetFuture();
 
 	// Enqueue the new task
-	FTask* Task = new FTask(NextTaskID++, CommandData, MoveTemp(Promise));
+	FDistributedBuildTask* Task = new FDistributedBuildTask(NextTaskID++, CommandData, MoveTemp(Promise));
 	{
 		PendingRequestedCompilationTasks.Enqueue(Task);
 	}
@@ -293,7 +293,7 @@ bool FUbaControllerModule::PollStats(FDistributedBuildStats& OutStats)
 	return JobDispatcherThread != nullptr && JobDispatcherThread->PollStats(OutStats);
 }
 
-void FUbaControllerModule::ReportJobProcessed(const FTaskResponse& InTaskResponse, FTask* CompileTask)
+void FUbaControllerModule::ReportJobProcessed(const FTaskResponse& InTaskResponse, FDistributedBuildTask* CompileTask)
 {
 	if (CompileTask)
 	{
diff --git a/Engine/Plugins/UbaController/Source/Private/UbaJobProcessor.cpp b/Engine/Plugins/UbaController/Source/Private/UbaJobProcessor.cpp
--- a/Engine/Plugins/UbaController/Source/Private/UbaJobProcessor.cpp
+++ b/Engine/Plugins/UbaController/Source/Private/UbaJobProcessor.cpp
@@ -247,7 +247,7 @@ void FUbaJobProcessor::CalculateKnownInputs()
 	KnownInputsBuffer.Add(0);
 }
 
-void FUbaJobProcessor::RunTaskWithUba(FTask* Task)
+void FUbaJobProcessor::RunTaskWithUba(FDistributedBuildTask* Task)
 {
 	FTaskCommandData& Data = Task->CommandData;
 	SessionServer_RegisterNewFile(UbaSessionServer, TCHAR_TO_UBASTRING(*Data.InputFileName));
@@ -286,7 +286,7 @@ void FUbaJobProcessor::RunTaskWithUba(FTask* Task)
 		FUbaJobProcessor* Processor;
 		FString InputFile;
 		FString OutputFile;
-		FTask* Task;
+		FDistributedBuildTask* Task;
 	};
 
 	auto Info = new ExitedInfo;
@@ -578,7 +578,7 @@ void FUbaJobProcessor::StartThread()
 	Thread = FRunnableThread::Create(this, TEXT("UbaJobProcessor"), 0, TPri_SlightlyBelowNormal, FPlatformAffinity::GetPoolThreadMask());
 }
 
-bool FUbaJobProcessor::ProcessOutputFile(FTask* CompileTask)
+bool FUbaJobProcessor::ProcessOutputFile(FDistributedBuildTask* CompileTask)
 {
 	//TODO: This method is mostly taken from the other Distribution controllers
 	// As we get an explicit callback when the process ends, we should be able to simplify this to just check if the file exists
@@ -631,7 +631,7 @@ bool FUbaJobProcessor::ProcessOutputFile(FTask* CompileTask)
 	return true;
 }
 
-void FUbaJobProcessor::HandleUbaJobFinished(FTask* CompileTask)
+void FUbaJobProcessor::HandleUbaJobFinished(FDistributedBuildTask* CompileTask)
 {
 	const bool bWasSuccessful = ProcessOutputFile(CompileTask);
 	if (!bWasSuccessful)
@@ -658,7 +658,7 @@ void FUbaJobProcessor::HandleTaskQueueUpdated(const FString& InputFileName)
 
 	while (true)
 	{
-		FTask* Task = nullptr;
+		FDistributedBuildTask* Task = nullptr;
 		if (!ControllerModule.PendingRequestedCompilationTasks.Dequeue(Task) || !Task)
 			break;
 		RunTaskWithUba(Task);
diff --git a/Engine/Plugins/UbaController/Source/Public/UbaControllerModule.h b/Engine/Plugins/UbaController/Source/Public/UbaControllerModule.h
--- a/Engine/Plugins/UbaController/Source/Public/UbaControllerModule.h
+++ b/Engine/Plugins/UbaController/Source/Public/UbaControllerModule.h
@@ -28,7 +28,7 @@ public:
 	virtual TFuture<FDistributedBuildTaskResult> EnqueueTask(const FTaskCommandData& CommandData) override final;
 	virtual bool PollStats(FDistributedBuildStats& OutStats) override;
 
-	void ReportJobProcessed(const FTaskResponse& InTaskResponse, FTask* CompileTask);
+	void ReportJobProcessed(const FTaskResponse& InTaskResponse, FDistributedBuildTask* CompileTask);
 	void CleanWorkingDirectory() const;
 
 	const FString& GetRootWorkingDirectory() const { return RootWorkingDirectory; }
@@ -38,7 +38,7 @@ public:
 	bool HasTasksDispatchedOrPending() const;
 
 	// Queue of tasks submitted by the engine, but not yet dispatched to the controller.
-	TQueue<FTask*, EQueueMode::Mpsc> PendingRequestedCompilationTasks;
+	TQueue<FDistributedBuildTask*, EQueueMode::Mpsc> PendingRequestedCompilationTasks;
 
 	static FString GetTempDir();
 
diff --git a/Engine/Plugins/UbaController/Source/Public/UbaJobProcessor.h b/Engine/Plugins/UbaController/Source/Public/UbaJobProcessor.h
--- a/Engine/Plugins/UbaController/Source/Public/UbaJobProcessor.h
+++ b/Engine/Plugins/UbaController/Source/Public/UbaJobProcessor.h
@@ -10,7 +10,7 @@
 
 class FUbaControllerModule;
 class FUbaHordeAgentManager;
-struct FTask;
+struct FDistributedBuildTask;
 
 class FUbaJobProcessor : public FRunnable, public TSharedFromThis<FUbaJobProcessor>
 {
@@ -27,9 +27,9 @@ public:
 	
 	/** Creates the threads and starts the main loop */
 	void StartThread();
-	bool ProcessOutputFile(FTask* CompileTask);
+	bool ProcessOutputFile(FDistributedBuildTask* CompileTask);
 
-	void HandleUbaJobFinished(FTask* CompileTask);
+	void HandleUbaJobFinished(FDistributedBuildTask* CompileTask);
 
 	/** Used to know when this thread has finished the main loop */
 	bool IsWorkDone() const { return bIsWorkDone;};
@@ -43,7 +43,7 @@ public:
 private:
 
 	void CalculateKnownInputs();
-	void RunTaskWithUba(FTask* Task);
+	void RunTaskWithUba(FDistributedBuildTask* Task);
 
 	void StartUba();
 	void ShutDownUba();
diff --git a/Engine/Plugins/XGEController/Source/Private/XGEControllerModule.cpp b/Engine/Plugins/XGEController/Source/Private/XGEControllerModule.cpp
--- a/Engine/Plugins/XGEController/Source/Private/XGEControllerModule.cpp
+++ b/Engine/Plugins/XGEController/Source/Private/XGEControllerModule.cpp
@@ -313,12 +313,12 @@ void FXGEControllerModule::ShutdownModule()
 		// Cancel any remaining tasks
 		for (auto Iter = DispatchedTasks.CreateIterator(); Iter; ++Iter)
 		{
-			FTask* Task = Iter.Value();
+			FDistributedBuildTask* Task = Iter.Value();
 			Task->Cancel();
 			delete Task;
 		}
 
-		FTask* Task;
+		FDistributedBuildTask* Task;
 		while (PendingTasks.Dequeue(Task))
 		{
 			Task->Cancel();
@@ -454,7 +454,7 @@ void FXGEControllerModule::WriteOutThreadProc()
 				break;
 
 			// Take one task from the pending queue.
-			FTask* Task = nullptr;
+			FDistributedBuildTask* Task = nullptr;
 			{
 				FScopeLock Lock(TasksCS);
 				PendingTasks.Dequeue(Task);
@@ -546,7 +546,7 @@ void FXGEControllerModule::ReadBackThreadProc()
 
 			// We've read a completed task response from the controller.
 			// Find the task in the map and complete the promise.
-			FTask* Task;
+			FDistributedBuildTask* Task;
 			{
 				FScopeLock Lock(TasksCS);
 				Task = DispatchedTasks.FindAndRemoveChecked(CompletedTaskResponse.ID);
@@ -576,7 +576,7 @@ TFuture<FDistributedBuildTaskResult> FXGEControllerModule::EnqueueTask(const FTa
 	TFuture<FDistributedBuildTaskResult> Future = Promise.GetFuture();
 
 	// Enqueue the new task
-	FTask* Task = new FTask(NextTaskID.Increment(), CommandData, MoveTemp(Promise));
+	FDistributedBuildTask* Task = new FDistributedBuildTask(NextTaskID.Increment(), CommandData, MoveTemp(Promise));
 	{
 		FScopeLock Lock(TasksCS);
 		PendingTasks.Enqueue(Task);
diff --git a/Engine/Plugins/XGEController/Source/Public/XGEControllerModule.h b/Engine/Plugins/XGEController/Source/Public/XGEControllerModule.h
--- a/Engine/Plugins/XGEController/Source/Public/XGEControllerModule.h
+++ b/Engine/Plugins/XGEController/Source/Public/XGEControllerModule.h
@@ -29,10 +29,10 @@ class FXGEControllerModule : public IDistributedBuildController
 	FCriticalSection* TasksCS;
 
 	// Queue of tasks submitted by the engine, but not yet dispatched to the controller.
-	TQueue<FTask*> PendingTasks;
+	TQueue<FDistributedBuildTask*> PendingTasks;
 
 	// Map of tasks dispatched to the controller and running within XGE, that have not yet finished.
-	TMap<uint32, FTask*> DispatchedTasks;
+	TMap<uint32, FDistributedBuildTask*> DispatchedTasks;
 
 	bool bShutdown;
 	bool bRestartWorker;
diff --git a/Engine/Source/Developer/DistributedBuildInterface/Public/DistributedBuildControllerInterface.h b/Engine/Source/Developer/DistributedBuildInterface/Public/DistributedBuildControllerInterface.h
--- a/Engine/Source/Developer/DistributedBuildInterface/Public/DistributedBuildControllerInterface.h
+++ b/Engine/Source/Developer/DistributedBuildInterface/Public/DistributedBuildControllerInterface.h
@@ -31,13 +31,13 @@ struct FTaskCommandData
 	TArray<FString> AdditionalOutputFolders; // Optional additional folder(s) which task may write artifacts to
 };
 
-struct FTask
+struct FDistributedBuildTask
 {
 	uint32 ID;
 	FTaskCommandData CommandData;
 	TPromise<FDistributedBuildTaskResult> Promise;
 
-	FTask(uint32 ID, const FTaskCommandData& CommandData, TPromise<FDistributedBuildTaskResult>&& Promise)
+	FDistributedBuildTask(uint32 ID, const FTaskCommandData& CommandData, TPromise<FDistributedBuildTaskResult>&& Promise)
 		: ID(ID)
 		, CommandData(CommandData)
 		, Promise(MoveTemp(Promise))
@@ -63,6 +63,8 @@ struct FTask
 	}
 };
 
+using FTask UE_DEPRECATED(5.6, "FTask from DistributedBUildControllerInterface.h has been renamed to FDistributedBuildTask in UE5.6") = FDistributedBuildTask;
+
 struct FTaskResponse
 {
 	uint32 ID;
