[Shaders] Disable FastBuildController plugin by default.

We use UbaController by default now and having FastBuildController enabled as fallback by default may cause issues if not configured correctly, so let users enable it explicitly.

#rnx

diff --git a/Engine/Plugins/FastBuildController/FastBuildController.uplugin b/Engine/Plugins/FastBuildController/FastBuildController.uplugin
--- a/Engine/Plugins/FastBuildController/FastBuildController.uplugin
+++ b/Engine/Plugins/FastBuildController/FastBuildController.uplugin
@@ -10,7 +10,7 @@
 	"DocsURL" : "",
 	"MarketplaceURL" : "",
 	"SupportURL" : "",
-	"EnabledByDefault" : true,
+	"EnabledByDefault" : false,
 	"CanContainContent" : false,
 	"IsBetaVersion" : false,
 	"Installed" : false,
